Many improvements to HTML5.
- Fixes to the Mac HTML5 Device selection. .app files now work correctly.
- Re-enabled HTML5 in Mac Editor.
- Added HTML5LaunchHelper executable to clean up the process of LaunchOn for HTML5.
- Improve HTML5 SDK Settings Editor interface. Only the emscripten install directory is needed now, SDK version are automatically picked up and selected for use.
- Change UnrealPak to also account for bytes saved (>64KB) and percentage size of original file (<90%) when choosing to automatically turn off compression.
- Added Server Port option for HTML5 deploy to stop clashes on port 8000
- Adding more logging for use during debugging & tracing.
- Added an option to turn on HTML5 tracing api and added calls to the api.
- Fix up check() macros to throw alert messages and be more clear that something has gone wrong on HTML5.

diff --git a/Engine/Config/BaseEngine.ini b/Engine/Config/BaseEngine.ini
--- a/Engine/Config/BaseEngine.ini
+++ b/Engine/Config/BaseEngine.ini
@@ -2226,4 +2226,5 @@ NameForGameView5="GameView5"
 ; These are the defaults for HTML5 settings, and they need to be in the .ini since UBT reads the .ini settings, without instantiating the class
 [/Script/HTML5PlatformEditor.HTML5TargetSettings]
 HeapSizeDevelopment=1024
-HeapSizeShipping=1024
\ No newline at end of file
+HeapSizeShipping=1024
+DeployServerPort=53501
\ No newline at end of file
diff --git a/Engine/Source/Developer/DesktopPlatform/Private/PlatformInfo.cpp b/Engine/Source/Developer/DesktopPlatform/Private/PlatformInfo.cpp
--- a/Engine/Source/Developer/DesktopPlatform/Private/PlatformInfo.cpp
+++ b/Engine/Source/Developer/DesktopPlatform/Private/PlatformInfo.cpp
@@ -106,7 +106,7 @@ static const FPlatformInfo PlatformInfoArray[] = {
 	BuildPlatformInfo(TEXT("Android_ES31"),				TEXT("Android_ES31"),		LOCTEXT("Android_ES31", "Android (ES31+AEP)"),				EPlatformType::Game,		EPlatformFlags::CookFlavor,		FPlatformIconPaths(TEXT("Launcher/AndroidTarget/Platform_Android_ES31_24x"), TEXT("Launcher/AndroidTarget/Platform_Android_128x")),					TEXT("-targetplatform=Android -cookflavor=ES31"),		TEXT("Android"),	EPlatformSDKStatus::Unknown,	TEXT("Shared/Tutorials/SettingUpAndroidTutorial"),														IsAvailableOnWindows,						false,				false),
 //	BuildPlatformInfo(TEXT("Android_GL4"),				TEXT("Android_GL4"),		LOCTEXT("Android_GL4", "Android (GL4)"),					EPlatformType::Game,		EPlatformFlags::CookFlavor,		FPlatformIconPaths(TEXT("Launcher/AndroidTarget/Platform_Android_GL4_24x"), TEXT("Launcher/AndroidTarget/Platform_Android_128x")),					TEXT("-targetplatform=Android -cookflavor=GL4"),		TEXT("Android"),	EPlatformSDKStatus::Unknown,	TEXT("Shared/Tutorials/SettingUpAndroidTutorial"),														IsAvailableOnWindows,						true),
 
-	BuildPlatformInfo(TEXT("HTML5"),					TEXT("HTML5"),				LOCTEXT("HTML5", "HTML5"),									EPlatformType::Game,		EPlatformFlags::None,			FPlatformIconPaths(TEXT("Launcher/HTML5Target/Platform_HTML5_24x"), TEXT("Launcher/HTML5Target/Platform_HTML5_128x")),								TEXT(""),												TEXT("HTML5"),		EPlatformSDKStatus::Unknown,	TEXT("Platforms/HTML5/GettingStarted"),																	IsAvailableOnWindows,   true,				false),
+	BuildPlatformInfo(TEXT("HTML5"),					TEXT("HTML5"),				LOCTEXT("HTML5", "HTML5"),									EPlatformType::Game,		EPlatformFlags::None,			FPlatformIconPaths(TEXT("Launcher/HTML5Target/Platform_HTML5_24x"), TEXT("Launcher/HTML5Target/Platform_HTML5_128x")),								TEXT(""),												TEXT("HTML5"),		EPlatformSDKStatus::Unknown,	TEXT("Platforms/HTML5/GettingStarted"),																	IsAvailableOnWindows || IsAvailableOnMac,   true,				false),
 
 	BuildPlatformInfo(TEXT("PS4"),						TEXT("PS4"),				LOCTEXT("PS4", "PlayStation 4"),							EPlatformType::Game,		EPlatformFlags::None,			FPlatformIconPaths(TEXT("Launcher/PS4Target/Platform_PS4_24x"), TEXT("Launcher/PS4Target/Platform_PS4_128x")),										TEXT(""),												TEXT("PS4"),		EPlatformSDKStatus::Unknown,	TEXT("Platforms/PS4/GettingStarted"),																	IsAvailableOnWindows,						false,				false),
 
diff --git a/Engine/Source/Developer/HTML5/HTML5PlatformEditor/Classes/HTML5TargetSettings.h b/Engine/Source/Developer/HTML5/HTML5PlatformEditor/Classes/HTML5TargetSettings.h
--- a/Engine/Source/Developer/HTML5/HTML5PlatformEditor/Classes/HTML5TargetSettings.h
+++ b/Engine/Source/Developer/HTML5/HTML5PlatformEditor/Classes/HTML5TargetSettings.h
@@ -31,4 +31,10 @@ public:
  	 */
  	UPROPERTY(GlobalConfig, EditAnywhere, Category=Memory, Meta = (DisplayName = "Heap Size (in MB)", ClampMin="1", ClampMax="4096"))
  	int32 HeapSizeShipping;
+
+	/**
+ 	 * Port to use when deploying game from the editor
+ 	 */
+ 	UPROPERTY(GlobalConfig, EditAnywhere, Category=Memory, Meta = (DisplayName = "Port to use when deploying game from the editor", ClampMin="49152", ClampMax="65535"))
+ 	int32 DeployServerPort;
 };
diff --git a/Engine/Source/Developer/HTML5/HTML5PlatformEditor/HTML5PlatformEditor.Build.cs b/Engine/Source/Developer/HTML5/HTML5PlatformEditor/HTML5PlatformEditor.Build.cs
--- a/Engine/Source/Developer/HTML5/HTML5PlatformEditor/HTML5PlatformEditor.Build.cs
+++ b/Engine/Source/Developer/HTML5/HTML5PlatformEditor/HTML5PlatformEditor.Build.cs
@@ -24,6 +24,9 @@ public class HTML5PlatformEditor : ModuleRules
 				"SourceControl",
                 "TargetPlatform",
 				"HTML5TargetPlatform",
+				// WITH_EDITOR only?
+				"UnrealEd",
+				"DesktopPlatform",
 			}
 		);
 
diff --git a/Engine/Source/Developer/HTML5/HTML5PlatformEditor/Private/HTML5PlatformEditorModule.cpp b/Engine/Source/Developer/HTML5/HTML5PlatformEditor/Private/HTML5PlatformEditorModule.cpp
--- a/Engine/Source/Developer/HTML5/HTML5PlatformEditor/Private/HTML5PlatformEditorModule.cpp
+++ b/Engine/Source/Developer/HTML5/HTML5PlatformEditor/Private/HTML5PlatformEditorModule.cpp
@@ -21,6 +21,10 @@ class FHTML5PlatformEditorModule
 	virtual void StartupModule() override
 	{
 		// register settings
+		static FName PropertyEditor("PropertyEditor");
+		FPropertyEditorModule& PropertyModule = FModuleManager::GetModuleChecked<FPropertyEditorModule>(PropertyEditor);
+		PropertyModule.RegisterCustomPropertyTypeLayout("HTML5SDKPath", FOnGetPropertyTypeCustomizationInstance::CreateStatic(&FHTML5SDKPathCustomization::MakeInstance));
+
 		ISettingsModule* SettingsModule = FModuleManager::GetModulePtr<ISettingsModule>("Settings");
 
 		if (SettingsModule != nullptr)
diff --git a/Engine/Source/Developer/HTML5/HTML5PlatformEditor/Private/HTML5PlatformEditorPrivatePCH.h b/Engine/Source/Developer/HTML5/HTML5PlatformEditor/Private/HTML5PlatformEditorPrivatePCH.h
--- a/Engine/Source/Developer/HTML5/HTML5PlatformEditor/Private/HTML5PlatformEditorPrivatePCH.h
+++ b/Engine/Source/Developer/HTML5/HTML5PlatformEditor/Private/HTML5PlatformEditorPrivatePCH.h
@@ -6,5 +6,9 @@
 #include "ModuleInterface.h"
 #include "ModuleManager.h"
 #include "Engine.h"
-
+// WITH_EDITOR (?)
+#include "UnrealEd.h"
+#include "PropertyEditing.h"
+#include "PropertyCustomizationHelpers.h"
+// endif WITH_EDITOR (?)
 #include "HTML5TargetSettings.h"
\ No newline at end of file
diff --git a/Engine/Source/Developer/HTML5/HTML5PlatformEditor/Private/HTML5SDKSettings.cpp b/Engine/Source/Developer/HTML5/HTML5PlatformEditor/Private/HTML5SDKSettings.cpp
--- a/Engine/Source/Developer/HTML5/HTML5PlatformEditor/Private/HTML5SDKSettings.cpp
+++ b/Engine/Source/Developer/HTML5/HTML5PlatformEditor/Private/HTML5SDKSettings.cpp
@@ -5,6 +5,13 @@
 //#include "EngineTypes.h"
 #include "HTML5SDKSettings.h"
 #include "IHTML5TargetPlatformModule.h"
+#include "PropertyHandle.h"
+#include "DesktopPlatformModule.h"
+#include "PlatformFilemanager.h"
+#include "GenericPlatformFile.h"
+#include "PlatformInfo.h"
+
+#define LOCTEXT_NAMESPACE "FHTML5PlatformEditorModule"
 
 DEFINE_LOG_CATEGORY_STATIC(HTML5SDKSettings, Log, All);
 
@@ -48,9 +55,9 @@ void UHTML5SDKSettings::QueryKnownBrowserLocations()
 		{TEXT("Firefox"), TEXT("C:/Program Files (x86)/Mozilla Firefox/firefox.exe")},
 		{TEXT("Chrome"), TEXT("C:/Program Files (x86)/Google/Chrome/Application/chrome.exe")},
 #elif PLATFORM_MAC
-		{TEXT("Safari"), TEXT("/Applications/Safari.app/Contents/MacOS/Safari")},
-		{TEXT("Firefox"),TEXT("/Applications/Firefox.app/Contents/MacOS/firefox-bin")},
-		{TEXT("Chrome"),TEXT("/Applications/Google Chrome.app/Contents/MacOS/Google Chrome")},
+		{TEXT("Safari"), TEXT("/Applications/Safari.app")},
+		{TEXT("Firefox"),TEXT("/Applications/Firefox.app")},
+		{TEXT("Chrome"),TEXT("/Applications/Google Chrome.app")},
 #elif PLATFORM_LINUX
 		{TEXT("Firefox"), TEXT("/usr/bin/firefox")},
 #else
@@ -70,4 +77,204 @@ void UHTML5SDKSettings::QueryKnownBrowserLocations()
 	}
 }
 
-#endif //WITH_EDITOR
\ No newline at end of file
+TSharedRef<IPropertyTypeCustomization> FHTML5SDKPathCustomization::MakeInstance()
+{
+	return MakeShareable(new FHTML5SDKPathCustomization());
+}
+
+void FHTML5SDKPathCustomization::CustomizeHeader(TSharedRef<IPropertyHandle> StructPropertyHandle, class FDetailWidgetRow& HeaderRow, IPropertyTypeCustomizationUtils& StructCustomizationUtils)
+{
+	PathProperty = StructPropertyHandle->GetChildHandle("SDKPath");
+	VersionProperty = StructPropertyHandle->GetChildHandle("EmscriptenVersion");
+
+	InvalidVersion = MakeShareable(new FHTML5Version());
+	InvalidVersion->Text = LOCTEXT("IncorrectSDKDirectory", "Incorrect SDK Directory");
+
+	UpdateAvailableVersions();
+
+	if (PathProperty.IsValid())
+	{
+		HeaderRow.ValueContent()
+			.MinDesiredWidth(125.0f)
+			.MaxDesiredWidth(600.0f)
+			[
+				SNew(SHorizontalBox)
+				+ SHorizontalBox::Slot()
+				.FillWidth(1.0f)
+				.VAlign(VAlign_Center)
+				[
+					SNew(SBox)
+					.MinDesiredWidth(200)
+					[
+						PathProperty->CreatePropertyValueWidget()
+					]
+				]
+				+ SHorizontalBox::Slot()
+					.AutoWidth()
+					.Padding(FMargin(4.0f, 0.0f, 0.0f, 0.0f))
+					.VAlign(VAlign_Center)
+					[
+						SAssignNew(BrowseButton, SButton)
+						.ButtonStyle(FEditorStyle::Get(), "HoverHintOnly")
+						.ToolTipText(LOCTEXT("FolderButtonToolTipText", "Choose a directory from this computer"))
+						.OnClicked(FOnClicked::CreateSP(this, &FHTML5SDKPathCustomization::OnPickDirectory, PathProperty.ToSharedRef()))
+						.ContentPadding(2.0f)
+						.ForegroundColor(FSlateColor::UseForeground())
+						.IsFocusable(false)
+						[
+							SNew(SImage)
+							.Image(FEditorStyle::GetBrush("PropertyWindow.Button_Ellipsis"))
+							.ColorAndOpacity(FSlateColor::UseForeground())
+						]
+					]
+				+ SHorizontalBox::Slot()
+					.AutoWidth()
+					.Padding(FMargin(4.0f, 0.0f, 0.0f, 0.0f))
+					.VAlign(VAlign_Center)
+					[
+						SAssignNew(SDKSelector, SComboBox< TSharedPtr<FHTML5Version> >)
+						.OptionsSource(&Options)
+						.InitiallySelectedItem(Options[0])
+						.OnSelectionChanged(this, &FHTML5SDKPathCustomization::OnSelectionChanged)
+						.ContentPadding(2)
+						.OnGenerateWidget(this, &FHTML5SDKPathCustomization::OnGenerateWidget)
+						[
+							SNew(STextBlock)
+							.Text(this, &FHTML5SDKPathCustomization::GetSelectedText)
+						]
+					]
+			]
+		.NameContent()
+			[
+				StructPropertyHandle->CreatePropertyNameWidget()
+			];
+	}
+}
+
+void FHTML5SDKPathCustomization::CustomizeChildren(TSharedRef<IPropertyHandle> StructPropertyHandle, class IDetailChildrenBuilder& StructBuilder, IPropertyTypeCustomizationUtils& StructCustomizationUtils)
+{
+}
+
+FReply FHTML5SDKPathCustomization::OnPickDirectory(TSharedRef<IPropertyHandle> PropertyHandle)// const
+{
+	FString Directory;
+	IDesktopPlatform* DesktopPlatform = FDesktopPlatformModule::Get();
+	auto& PlatformFile = FPlatformFileManager::Get().GetPlatformFile();
+	auto* Settings = GetMutableDefault<UHTML5SDKSettings>();
+	if (!DesktopPlatform || !Settings)
+	{
+		return FReply::Handled();
+	}
+
+	TSharedPtr<SWindow> ParentWindow = FSlateApplication::Get().FindWidgetWindow(BrowseButton.ToSharedRef());
+	void* ParentWindowHandle = (ParentWindow.IsValid() && ParentWindow->GetNativeWindow().IsValid()) ? ParentWindow->GetNativeWindow()->GetOSWindowHandle() : nullptr;
+
+	FString StartDirectory = FEditorDirectories::Get().GetLastDirectory(ELastDirectory::GENERIC_IMPORT);
+	for (;;)
+	{
+		if (DesktopPlatform->OpenDirectoryDialog(ParentWindowHandle, LOCTEXT("FolderDialogTitle", "Choose a directory").ToString(), StartDirectory, Directory))
+		{
+			FText FailureReason;
+ 			if (PlatformFile.DirectoryExists(*Directory))
+ 			{
+				// TODO: Update config value too
+				Settings->EmscriptenRoot.SDKPath = Directory;
+				PathProperty->SetValue(Directory);
+				UpdateAvailableVersions();
+				SDKSelector->SetSelectedItem(SelectedVersionItem);
+				SDKSelector->RefreshOptions();
+ 			}
+ 			else
+			{
+				StartDirectory = Directory;
+				FMessageDialog::Open(EAppMsgType::Ok, FailureReason);
+				continue;
+			}
+		}
+		break;
+	}
+
+	return FReply::Handled();
+}
+
+TSharedRef<SWidget> FHTML5SDKPathCustomization::OnGenerateWidget(TSharedPtr<FHTML5Version> InItem)
+{
+	return SNew(STextBlock)
+		.Text(InItem->Text);
+}
+
+FText FHTML5SDKPathCustomization::GetSelectedText() const
+{
+	return SelectedVersionItem->Text;
+}
+
+void FHTML5SDKPathCustomization::OnSelectionChanged(TSharedPtr<FHTML5Version> InItem, ESelectInfo::Type InSeletionInfo)
+{
+	SelectedVersionItem = InItem;
+	auto* Settings = GetMutableDefault<UHTML5SDKSettings>();
+	VersionProperty->SetValue(InItem->Version.VersionNumberAsString());
+}
+
+void FHTML5SDKPathCustomization::UpdateAvailableVersions()
+{
+	Options.Empty();
+	Options.Add(InvalidVersion);
+
+	//Take the cheap option here. SDK is now uninstalled. Don't query UBT!
+	// If we clear this function we're good to install again
+	PlatformInfo::UpdatePlatformSDKStatus(TEXT("HTML5"), PlatformInfo::EPlatformSDKStatus::NotInstalled);
+
+	SelectedVersionItem = Options[0];
+
+	auto* Settings = GetMutableDefault<UHTML5SDKSettings>();
+	auto& PlatformFile = FPlatformFileManager::Get().GetPlatformFile();
+	TargetPlatformModule = FModuleManager::LoadModulePtr<IHTML5TargetPlatformModule>("HTML5TargetPlatform");
+	if (!Settings || !TargetPlatformModule)
+	{
+		return;
+	}
+
+	TSharedPtr<FHTML5Version> PrevSelectedVersion(new FHTML5Version);
+	PrevSelectedVersion->Version.VersionNumberFromString(*Settings->EmscriptenRoot.EmscriptenVersion);
+
+	TArray<FHTML5SDKVersionNumber> SDKVersions;
+	TargetPlatformModule->GetInstalledSDKVersions(*(Settings->EmscriptenRoot.SDKPath), SDKVersions);
+
+	if (SDKVersions.Num() == 0)
+	{
+		return;
+	}
+	Options.Empty();
+	for (const auto& i : SDKVersions)
+	{
+		TSharedPtr<FHTML5Version> Ver(new FHTML5Version);
+		Ver->Text = FText::FromString(i.VersionNumberAsString());
+		Ver->Version = i;
+		Options.Add(Ver);
+	}
+
+	TSharedPtr<FHTML5Version> Latest(new FHTML5Version);
+	*Latest = *Options.Last();
+	Latest->Text = LOCTEXT("LastestVesion", "Latest");
+	Latest->Version.MakeLatestVersionNumber();
+	Options.Add(Latest);
+
+	int32 FoundIndex = Options.IndexOfByPredicate([&](const TSharedPtr<FHTML5Version>& RHS)
+	{
+		return RHS->Version == PrevSelectedVersion->Version;
+	});
+	if (FoundIndex != INDEX_NONE)
+	{
+		SelectedVersionItem = Options[FoundIndex];
+	}
+	else
+	{
+		SelectedVersionItem = Options.Last();
+	}
+
+	//Take the cheap option here. If we got here, there is a valid install selected. 
+	// Don't query UBT, it'll repeat the same test via a separate process! :(
+	PlatformInfo::UpdatePlatformSDKStatus(TEXT("HTML5"), PlatformInfo::EPlatformSDKStatus::Installed);
+}
+#endif //WITH_EDITOR
+#undef LOCTEXT_NAMESPACE
\ No newline at end of file
diff --git a/Engine/Source/Developer/HTML5/HTML5PlatformEditor/Private/HTML5SDKSettings.h b/Engine/Source/Developer/HTML5/HTML5PlatformEditor/Private/HTML5SDKSettings.h
--- a/Engine/Source/Developer/HTML5/HTML5PlatformEditor/Private/HTML5SDKSettings.h
+++ b/Engine/Source/Developer/HTML5/HTML5PlatformEditor/Private/HTML5SDKSettings.h
@@ -3,6 +3,10 @@
 #pragma once
 
 #include "ITargetPlatformManagerModule.h"
+#include "EditorStyleSet.h"
+#include "IDetailCustomization.h"
+#include "IPropertyTypeCustomization.h"
+#include "IHTML5TargetPlatformModule.h"
 #include "HTML5SDKSettings.generated.h"
 
 class IHTML5TargetPlatformModule;
@@ -19,6 +23,19 @@ struct FHTML5DeviceMapping
 	FFilePath DevicePath;
 };
 
+USTRUCT()
+struct FHTML5SDKPath
+{
+	GENERATED_USTRUCT_BODY()
+
+	UPROPERTY(EditAnywhere, Category = SDK_Path)
+	FString SDKPath;
+
+	// Note that 0.0.0 is invalid and -1.-1.-1 is latest
+	UPROPERTY(EditAnywhere, Category = SDK_Path)
+	FString EmscriptenVersion;
+};
+
 /**
  * Implements the settings for the HTML5 SDK setup.
  */
@@ -28,9 +45,6 @@ class HTML5PLATFORMEDITOR_API UHTML5SDKSettings : public UObject
 public:
 	GENERATED_UCLASS_BODY()
 
-	UPROPERTY(GlobalConfig, EditAnywhere, Category = HTML5_SDK_Paths, Meta = (DisplayName = "Location of Emscripten SDK (This directory contains 'emcc')"))
-	FDirectoryPath Emscripten;
-
 	UPROPERTY(GlobalConfig, EditAnywhere, Category = HTML5_SDK_Paths, Meta = (DisplayName = "Location of Python exe (This can be blank if python is on the PATH)"))
 	FFilePath Python;
 
@@ -38,12 +52,72 @@ public:
 	UPROPERTY(GlobalConfig, EditAnywhere, Category = HTML5_Devices, Meta = (DisplayName = "Available browsers"))
 	TArray<FHTML5DeviceMapping> DeviceMap;
 
+	// Path to Emscripten SDK install directory
+	UPROPERTY(GlobalConfig, EditAnywhere, Category = HTML5_SDK_Paths, Meta = (DisplayName = "Location of Emscripten SDK (This directory contains all emscripten versions)"))
+	FHTML5SDKPath EmscriptenRoot;
+
 #if WITH_EDITOR
 	// UObject interface
 	virtual void PostEditChangeProperty(struct FPropertyChangedEvent& PropertyChangedEvent) override;
 	// End of UObject interface
 	void QueryKnownBrowserLocations();
+
+	ITargetPlatformManagerModule * TargetManagerModule;
 #endif
 private:
 	IHTML5TargetPlatformModule* TargetPlatformModule;
 };
+
+#if WITH_EDITOR
+struct FHTML5Version
+{
+	// Version text to display
+	FText Text;
+
+	FHTML5SDKVersionNumber Version;
+
+	bool operator == (const FHTML5Version& RHS) const
+	{
+		return Version == RHS.Version;
+	}
+};
+
+class FHTML5SDKPathCustomization : public IPropertyTypeCustomization
+{
+public:
+	static TSharedRef<IPropertyTypeCustomization> MakeInstance();
+
+	// IPropertyTypeCustomization interface 
+	virtual void CustomizeHeader(TSharedRef<class IPropertyHandle> StructPropertyHandle, class FDetailWidgetRow& HeaderRow, IPropertyTypeCustomizationUtils& StructCustomizationUtils) override;
+
+	virtual void CustomizeChildren(TSharedRef<class IPropertyHandle> StructPropertyHandle, class IDetailChildrenBuilder& StructBuilder, IPropertyTypeCustomizationUtils& StructCustomizationUtils) override;
+
+private:
+
+	FReply OnPickDirectory(TSharedRef<IPropertyHandle> PropertyHandle);// const;
+
+	TSharedRef<SWidget> OnGenerateWidget(TSharedPtr<FHTML5Version> InItem);
+
+	void OnSelectionChanged(TSharedPtr<FHTML5Version> InItem, ESelectInfo::Type InSeletionInfo);
+
+	FText GetSelectedText() const;
+
+	void UpdateAvailableVersions();
+
+	/** The browse button widget */
+	TSharedPtr<SButton> BrowseButton;
+
+	TSharedPtr< SComboBox< TSharedPtr<FHTML5Version> > > SDKSelector;
+
+	TArray< TSharedPtr<FHTML5Version> > Options;
+
+	TSharedPtr< FHTML5Version > SelectedVersionItem;
+
+	TSharedPtr< FHTML5Version > InvalidVersion;
+
+	IHTML5TargetPlatformModule* TargetPlatformModule;
+
+	TSharedPtr<IPropertyHandle> PathProperty;
+	TSharedPtr<IPropertyHandle> VersionProperty;
+};
+#endif // WITH_EDITOR
diff --git a/Engine/Source/Developer/HTML5/HTML5TargetPlatform/Private/HTML5TargetPlatform.cpp b/Engine/Source/Developer/HTML5/HTML5TargetPlatform/Private/HTML5TargetPlatform.cpp
--- a/Engine/Source/Developer/HTML5/HTML5TargetPlatform/Private/HTML5TargetPlatform.cpp
+++ b/Engine/Source/Developer/HTML5/HTML5TargetPlatform/Private/HTML5TargetPlatform.cpp
@@ -15,9 +15,6 @@
 
 FHTML5TargetPlatform::FHTML5TargetPlatform( )
 {
-	// load the final HTML5 engine settings for this game
-	FConfigCacheIni::LoadLocalIniFile(HTML5EngineSettings, TEXT("Engine"), true, *PlatformName());
-
 	RefreshAvailableDevices();
 
 #if WITH_ENGINE
@@ -69,17 +66,67 @@ bool FHTML5TargetPlatform::IsSdkInstalled(bool bProjectHasCode, FString& OutDocu
 	//New style detection of devices
 	if (HTML5EngineSettings.Find("/Script/HTML5PlatformEditor.HTML5SDKSettings"))
 	{
+
 		SDKPaths = HTML5EngineSettings["/Script/HTML5PlatformEditor.HTML5SDKSettings"];
 		for (auto It : SDKPaths)
 		{
 			const FString& Platform = It.Key.ToString();
 			FString Path = It.Value;
+			if (Platform == "EmscriptenRoot")
 			{
-				Path.RemoveFromStart(TEXT("(Path=\""));
-				Path.RemoveFromEnd(TEXT("\")"));
-				if (Platform == "Emscripten" && IFileManager::Get().DirectoryExists(*Path))
+				FString SDKPathString;
+				FString VersionString;
+				auto Index = Path.Find(TEXT("SDKPath="));
+				if (Index != INDEX_NONE)
 				{
-					return true;
+					Index+=9;
+					auto Index2 = Path.Find(TEXT("\""), ESearchCase::IgnoreCase, ESearchDir::FromStart, Index);
+					if (Index2 != INDEX_NONE)
+					{
+						SDKPathString = Path.Mid(Index, Index2-Index);
+					}
+				}
+				Index = Path.Find(TEXT("EmscriptenVersion="));
+				if (Index != INDEX_NONE)
+				{
+					Index += 19;
+					auto Index2 = Path.Find(TEXT("\""), ESearchCase::IgnoreCase, ESearchDir::FromStart, Index);
+					if (Index2 != INDEX_NONE)
+					{
+						VersionString = Path.Mid(Index, Index2 - Index);
+					}
+				}
+				if (!VersionString.IsEmpty() && !SDKPathString.IsEmpty())
+				{
+					bool bFoundSDK = false;
+					FHTML5SDKVersionNumber RequiredVersion;
+					RequiredVersion.VersionNumberFromString(*VersionString);
+					TArray<FHTML5SDKVersionNumber> SDKVersions;
+					GetInstalledSDKVersions(*SDKPathString, SDKVersions);
+
+					if (SDKVersions.Num() > 0 &&
+						RequiredVersion.VersionNumber.Major == -1 &&
+						RequiredVersion.VersionNumber.Minor == -1 &&
+						RequiredVersion.VersionNumber.Revision == -1)
+					{
+						bFoundSDK = true;
+						SDKPathString = SDKVersions.Last().VersionPath;
+					}
+					else
+					{
+						for (const auto& i : SDKVersions)
+						{
+							if (i == RequiredVersion)
+							{
+								SDKPathString = i.VersionPath;
+								bFoundSDK = true;
+								break;
+							}
+						}
+					}
+
+					
+					return bFoundSDK && IFileManager::Get().DirectoryExists(*SDKPathString);
 				}
 			}
 		}
@@ -278,6 +325,9 @@ FName FHTML5TargetPlatform::GetWaveFormat( const USoundWave* Wave ) const
 
 void FHTML5TargetPlatform::RefreshAvailableDevices()
 {
+	// load the final HTML5 engine settings for this game
+	FConfigCacheIni::LoadLocalIniFile(HTML5EngineSettings, TEXT("Engine"), true, *PlatformName());
+
 	//New style detection of devices
 	for (const auto& Device : LocalDevice)
 	{
@@ -319,7 +369,8 @@ void FHTML5TargetPlatform::RefreshAvailableDevices()
 				}
 
 				if (!DeviceName.IsEmpty() && !DevicePath.IsEmpty() &&
-					FPlatformFileManager::Get().GetPlatformFile().FileExists(*DevicePath))
+					(FPlatformFileManager::Get().GetPlatformFile().FileExists(*DevicePath) ||
+                     FPlatformFileManager::Get().GetPlatformFile().DirectoryExists(*DevicePath)))
 				{
 					ITargetDevicePtr Device = MakeShareable(new FHTML5TargetDevice(*this, FString::Printf(TEXT("%s on %s"), *DeviceName, FPlatformProcess::ComputerName())));
 					LocalDevice.Add(Device);
@@ -358,3 +409,49 @@ void FHTML5TargetPlatform::RefreshAvailableDevices()
 		}
 	}
 }
+
+void FHTML5TargetPlatform::GetInstalledSDKVersions(const TCHAR* SDKDirectory, TArray<FHTML5SDKVersionNumber>& OutSDKs)
+{
+	if (!SDKDirectory)
+	{
+		return;
+	}
+
+	struct FVersionSearch : public IPlatformFile::FDirectoryVisitor
+	{
+		IPlatformFile* PlatformFile;
+		TArray<FHTML5SDKVersionNumber>& SDKs;
+		FString VersionFilePath;
+
+		FVersionSearch(IPlatformFile* InPlatformFile, TArray<FHTML5SDKVersionNumber>& InSDKs)
+			: PlatformFile(InPlatformFile)
+			, SDKs(InSDKs)
+		{}
+		virtual bool Visit(const TCHAR* FilenameOrDirectory, bool bIsDirectory)
+		{
+			if (bIsDirectory == true)
+			{
+				VersionFilePath = FPaths::Combine(FilenameOrDirectory, TEXT("emscripten-version.txt"));
+				if (PlatformFile->FileExists(*VersionFilePath))
+				{
+					FString VersionText;
+					FFileHelper::LoadFileToString(VersionText, *VersionFilePath);
+					VersionText.Trim();
+					VersionText.TrimTrailing();
+					FHTML5SDKVersionNumber Ver;
+					Ver.VersionPath = FilenameOrDirectory;
+					Ver.VersionNumberFromString(*VersionText);
+					SDKs.Add(Ver);
+				}
+			}
+			return true;
+		}
+	};
+
+	auto& PlatformFile = FPlatformFileManager::Get().GetPlatformFile();
+
+	FVersionSearch VersionSearch(&PlatformFile, OutSDKs);
+	FString EmscriptenSDkDir = FPaths::Combine(SDKDirectory, TEXT("emscripten"));
+	PlatformFile.IterateDirectory(*EmscriptenSDkDir, VersionSearch);
+	OutSDKs.Sort();
+}
diff --git a/Engine/Source/Developer/HTML5/HTML5TargetPlatform/Private/HTML5TargetPlatform.h b/Engine/Source/Developer/HTML5/HTML5TargetPlatform/Private/HTML5TargetPlatform.h
--- a/Engine/Source/Developer/HTML5/HTML5TargetPlatform/Private/HTML5TargetPlatform.h
+++ b/Engine/Source/Developer/HTML5/HTML5TargetPlatform/Private/HTML5TargetPlatform.h
@@ -89,6 +89,8 @@ public:
 
 	// End ITargetPlatform interface
 
+	static void GetInstalledSDKVersions(const TCHAR* SDKDirectory, TArray<FHTML5SDKVersionNumber>& OutSDKs);
+
 private:
 
 	// Holds the HTML5 engine settings.
diff --git a/Engine/Source/Developer/HTML5/HTML5TargetPlatform/Private/HTML5TargetPlatformModule.cpp b/Engine/Source/Developer/HTML5/HTML5TargetPlatform/Private/HTML5TargetPlatformModule.cpp
--- a/Engine/Source/Developer/HTML5/HTML5TargetPlatform/Private/HTML5TargetPlatformModule.cpp
+++ b/Engine/Source/Developer/HTML5/HTML5TargetPlatform/Private/HTML5TargetPlatformModule.cpp
@@ -59,6 +59,11 @@ public:
 		}
 	}
 
+	virtual void GetInstalledSDKVersions(const TCHAR* SDKDirectory, TArray<FHTML5SDKVersionNumber>& OutSDKs)
+	{
+		FHTML5TargetPlatform::GetInstalledSDKVersions(SDKDirectory, OutSDKs);
+	}
+
 public:
 
 	// Begin IModuleInterface interface
diff --git a/Engine/Source/Developer/HTML5/HTML5TargetPlatform/Public/IHTML5TargetPlatformModule.h b/Engine/Source/Developer/HTML5/HTML5TargetPlatform/Public/IHTML5TargetPlatformModule.h
--- a/Engine/Source/Developer/HTML5/HTML5TargetPlatform/Public/IHTML5TargetPlatformModule.h
+++ b/Engine/Source/Developer/HTML5/HTML5TargetPlatform/Public/IHTML5TargetPlatformModule.h
@@ -4,8 +4,78 @@
 
 #include "ITargetPlatformModule.h"
 
+class HTML5TARGETPLATFORM_API FHTML5SDKVersionNumber
+{
+public:
+	FHTML5SDKVersionNumber()
+	{
+		// Zero is invalid, -1 is latest.
+		VersionNumber.Major = 0;
+		VersionNumber.Minor = 0;
+		VersionNumber.Revision = 0;
+	}
+
+	// Directory of the SDK 
+	FString VersionPath;
+
+	//Version Number
+	struct
+	{
+		int32 Major;
+		int32 Minor;
+		int32 Revision;
+	} VersionNumber;
+
+	void VersionNumberFromString(const TCHAR* VersionString)
+	{
+		// We expect the format "1.23.45"
+		const TCHAR* VStr = VersionString;
+		TCHAR* VEnd = nullptr;
+		int32 Nums[3] = { 0 };
+		VersionNumber.Major = 0;
+		VersionNumber.Minor = 0;
+		VersionNumber.Revision = 0;
+		for (auto& N : Nums)
+		{
+			N = FCString::Strtoi(VStr, &VEnd, 10);
+			if (!VEnd || VStr == VEnd || (*VEnd != '.' && *VEnd != 0))
+			{
+				return;
+			}
+			VStr = VEnd + 1;
+		}
+		VersionNumber.Major = Nums[0];
+		VersionNumber.Minor = Nums[1];
+		VersionNumber.Revision = Nums[2];
+	}
+
+	FString VersionNumberAsString() const
+	{
+		return FString::Printf(TEXT("%d.%d.%d"), VersionNumber.Major, VersionNumber.Minor, VersionNumber.Revision);
+	}
+
+	void MakeLatestVersionNumber()
+	{
+		VersionNumber.Major = -1;
+		VersionNumber.Minor = -1;
+		VersionNumber.Revision = -1;
+	}
+
+	bool operator < (const FHTML5SDKVersionNumber& RHS) const
+	{
+		return (VersionNumber.Major < RHS.VersionNumber.Major) ||
+			(VersionNumber.Major == RHS.VersionNumber.Major && VersionNumber.Minor < RHS.VersionNumber.Minor) ||
+			((VersionNumber.Major == RHS.VersionNumber.Major) && (VersionNumber.Minor == RHS.VersionNumber.Minor) && (VersionNumber.Revision < RHS.VersionNumber.Revision));
+	}
+
+	bool operator == (const FHTML5SDKVersionNumber& RHS) const
+	{
+		return VersionNumber.Major == RHS.VersionNumber.Major && VersionNumber.Minor == RHS.VersionNumber.Minor && VersionNumber.Revision == RHS.VersionNumber.Revision;
+	}
+};
+
 /**
- * Interface for AndroidDeviceDetection module.
+ * Interface for HTML5TargetPlatformModule module.
  */
 class IHTML5TargetPlatformModule
 	: public ITargetPlatformModule
@@ -16,6 +86,11 @@ public:
 	 */
 	virtual void RefreshAvailableDevices() = 0;
 
+	/**
+	 * Fills a list of installed SDK version numbers
+	 */
+	virtual void GetInstalledSDKVersions(const TCHAR* SDKDirectory, TArray<FHTML5SDKVersionNumber>& OutSDKs) = 0;
+
 protected:
 
 	/**
diff --git a/Engine/Source/Programs/AutomationTool/HTML5/HTML5Platform.Automation.cs b/Engine/Source/Programs/AutomationTool/HTML5/HTML5Platform.Automation.cs
--- a/Engine/Source/Programs/AutomationTool/HTML5/HTML5Platform.Automation.cs
+++ b/Engine/Source/Programs/AutomationTool/HTML5/HTML5Platform.Automation.cs
@@ -259,8 +259,11 @@ public class HTML5Platform : Platform
 		}
 
 		// open the webpage
-		string directory = Path.GetDirectoryName(ClientApp);
+		Int32 ServerPort = 8000;
+		ConfigCache.GetInt32("/Script/HTML5PlatformEditor.HTML5TargetSettings", "DeployServerPort", out ServerPort);
+		string WorkingDirectory = Path.GetDirectoryName(ClientApp);
 		string url = Path.GetFileName(ClientApp) +".html";
+		string args = "-m ";
 		// Are we running via cook on the fly server?
 		// find our http url - This is awkward because RunClient doesn't have real information that NFS is running or not.
 		bool IsCookOnTheFly = false;
@@ -281,100 +284,27 @@ public class HTML5Platform : Platform
 		}
 		else
 		{
-			url = "http://127.0.0.1:8000/" + url;
-			// this will be killed UBT instances dies.
-			string input = String.Format(" -m SimpleHTTPServer 8000");
-
-            string PythonPath = HTML5SDKInfo.PythonPath();
-            ProcessResult Result = ProcessManager.CreateProcess(PythonPath, true, "html5server.log");
-            Result.ProcessObject.StartInfo.FileName = PythonPath;
-			Result.ProcessObject.StartInfo.UseShellExecute = false;
-			Result.ProcessObject.StartInfo.RedirectStandardOutput = true;
-			Result.ProcessObject.StartInfo.RedirectStandardInput = true;
-			Result.ProcessObject.StartInfo.WorkingDirectory = directory;
-			Result.ProcessObject.StartInfo.Arguments = input;
-			Result.ProcessObject.Start();
-
-			Result.ProcessObject.OutputDataReceived += delegate(object sender, System.Diagnostics.DataReceivedEventArgs e)
-			{
-				System.Console.WriteLine(e.Data);
-			};
-
-			System.Console.WriteLine("Starting Browser Process");
-
-
-			// safari specific hack.
-			string argument = url; 
-			if (browserPath.Contains ("Safari") && Utils.IsRunningOnMono)
-				argument = "";
-
-			// Chrome issue. Firefox may work like this in the future
-			bool bBrowserWillSpawnProcess = browserPath.Contains("chrome") || (browserPath.Contains("Google Chrome") && Utils.IsRunningOnMono);
-
-			ProcessResult SubProcess = null;
-			ProcessResult ClientProcess = Run(browserPath, argument, null, ClientRunFlags | ERunOptions.NoWaitForExit);
-			var ProcStartTime = ClientProcess.ProcessObject.StartTime;
-			var ProcName = ClientProcess.ProcessObject.ProcessName;
-
-			ClientProcess.ProcessObject.EnableRaisingEvents = true;
-			ClientProcess.ProcessObject.Exited += delegate(System.Object o, System.EventArgs e)
-			{
-				System.Console.WriteLine("Browser Process Ended (PID={0})", ClientProcess.ProcessObject.Id);
-				var bFoundChildProcess = false;
-				if (bBrowserWillSpawnProcess)
-				{
-					// Chrome spawns a process from the tab it opens and then lets the process we spawned die, so
-					// catch that process and attach to that instead.
-					var CurrentProcesses = Process.GetProcesses();
-					foreach (var item in CurrentProcesses)
-					{
-						if (item.Id != ClientProcess.ProcessObject.Id && item.ProcessName == ProcName && item.StartTime >= ProcStartTime && item.StartTime <= ClientProcess.ProcessObject.ExitTime)
-						{
-							var PID = item.Id;
-							System.Console.WriteLine("Found Process {0} with PID {1} which started at {2}. Waiting on that process to end.", item.ProcessName, item.Id, item.StartTime.ToString());
-							SubProcess = new ProcessResult(item.ProcessName, item, true, item.ProcessName);
-							item.EnableRaisingEvents = true;
-							item.Exited += delegate(System.Object o2, System.EventArgs e2)
-							{
-								System.Console.WriteLine("Browser Process Ended (PID={0}) - Killing Webserver", PID);
-								Result.ProcessObject.StandardInput.Close();
-								Result.ProcessObject.Kill();
-							};
-							bFoundChildProcess = true;
-						}
-					}
-				}
-
-				if (!bFoundChildProcess)
-				{
-					System.Console.WriteLine("- Killing Webserver PID({0})", Result.ProcessObject.Id);
-					Result.ProcessObject.StandardInput.Close();
-					Result.ProcessObject.Kill();
-				}
-			};
-
-			if (bBrowserWillSpawnProcess)
-			{
-				//Wait for it to do so...
-				ClientProcess.ProcessObject.WaitForExit();
-				ClientProcess = SubProcess;
-			}
-
-			// safari needs a hack. 
-			// http://superuser.com/questions/689315/run-safari-from-terminal-with-given-url-address-without-open-command
-
-			if (browserPath.Contains ("Safari") && Utils.IsRunningOnMono) 
-			{
-				// ClientProcess.ProcessObject.WaitForInputIdle ();
-				Thread.Sleep (2000);
-				Process.Start("/usr/bin/osascript"," -e 'tell application \"Safari\" to open location \"" + url + "\"'" );
-			}
+			url = String.Format("http://127.0.0.1:{0}/{1}", ServerPort, url);
+			args += String.Format("-h -s {0} ", ServerPort);
+		}
 
-			return ClientProcess;
+		// Check HTML5LaunchHelper source for command line args
+		var LowerBrowserPath = browserPath.ToLower();
+		args += String.Format("-b \"{0}\" -p \"{1}\" -w \"{2}\" ", browserPath, HTML5SDKInfo.PythonPath(), WorkingDirectory);
+		args += url + " ";
+		var ProfileDirectory = Path.Combine(Utils.GetUserSettingDirectory(), "UE4_HTML5", "user");
+		if (LowerBrowserPath.Contains("chrome"))
+		{
+			args += String.Format("--user-data-dir=\"{0}\" --enable-logging --no-first-run", Path.Combine(ProfileDirectory, "chrome"));
+		}
+		else if (LowerBrowserPath.Contains("firefox"))
+		{
+			args += String.Format("-no-remote -profile \"{0}\" -jsconsole", Path.Combine(ProfileDirectory, "firefox"));
 		}
+		//else if (browserPath.Contains("Safari")) {}
 
-		System.Console.WriteLine("Browser Path " + browserPath);
-		ProcessResult BrowserProcess = Run(browserPath, url, null, ClientRunFlags | ERunOptions.NoWaitForExit);
+		var LaunchHelperPath = CombinePaths(CmdEnv.LocalRoot, "Engine/Binaries/DotNET/HTML5LaunchHelper.exe");
+		ProcessResult BrowserProcess = Run(LaunchHelperPath, args, null, ClientRunFlags | ERunOptions.NoWaitForExit);
 	
 		return BrowserProcess;
 
diff --git a/Engine/Source/Programs/AutomationTool/Scripts/BuildCommonTools.Automation.cs b/Engine/Source/Programs/AutomationTool/Scripts/BuildCommonTools.Automation.cs
--- a/Engine/Source/Programs/AutomationTool/Scripts/BuildCommonTools.Automation.cs
+++ b/Engine/Source/Programs/AutomationTool/Scripts/BuildCommonTools.Automation.cs
@@ -147,6 +147,15 @@ public class BuildCommonTools : BuildCommand
 		{
 			Agenda.AddTarget("XboxOnePDBFileUtil", UnrealBuildTool.UnrealTargetPlatform.Win64, UnrealBuildTool.UnrealTargetConfiguration.Development);
 		}
+
+		// HTML5 binaries
+		if (Platforms.Contains(UnrealBuildTool.UnrealTargetPlatform.HTML5))
+		{
+			Agenda.AddTarget("UnrealPak", UnrealBuildTool.UnrealTargetPlatform.Win64, UnrealBuildTool.UnrealTargetConfiguration.Development);
+			
+			Agenda.DotNetProjects.Add(@"Engine/Source/Programs/HTML5LaunchHelper/HTML5LaunchHelper.csproj");
+			ExtraBuildProducts.Add(CommandUtils.CombinePaths(CommandUtils.CmdEnv.LocalRoot, @"Engine/Binaries/DotNET/HTML5/HTML5LaunchHelper.exe"));
+		}
 		
 		return Agenda;
 	}
diff --git a/Engine/Source/Programs/HTML5/HTML5LaunchHelper/App.config b/Engine/Source/Programs/HTML5/HTML5LaunchHelper/App.config
new file mode 100644
--- /dev/null
+++ b/Engine/Source/Programs/HTML5/HTML5LaunchHelper/App.config
@@ -0,0 +1,6 @@
+﻿<?xml version="1.0" encoding="utf-8" ?>
+<configuration>
+    <startup> 
+        <supportedRuntime version="v4.0" sku=".NETFramework,Version=v4.5" />
+    </startup>
+</configuration>
\ No newline at end of file
diff --git a/Engine/Source/Programs/HTML5/HTML5LaunchHelper/HTML5LaunchHelper.csproj b/Engine/Source/Programs/HTML5/HTML5LaunchHelper/HTML5LaunchHelper.csproj
new file mode 100644
--- /dev/null
+++ b/Engine/Source/Programs/HTML5/HTML5LaunchHelper/HTML5LaunchHelper.csproj
@@ -0,0 +1,65 @@
+﻿<?xml version="1.0" encoding="utf-8"?>
+<Project ToolsVersion="12.0" DefaultTargets="Build" xmlns="http://schemas.microsoft.com/developer/msbuild/2003">
+  <Import Project="$(MSBuildExtensionsPath)\$(MSBuildToolsVersion)\Microsoft.Common.props" Condition="Exists('$(MSBuildExtensionsPath)\$(MSBuildToolsVersion)\Microsoft.Common.props')" />
+  <PropertyGroup>
+    <Configuration Condition=" '$(Configuration)' == '' ">Debug</Configuration>
+    <Platform Condition=" '$(Platform)' == '' ">AnyCPU</Platform>
+    <ProjectGuid>{C22FE09B-9D62-4B4C-AE2E-BDB5D66C143F}</ProjectGuid>
+    <OutputType>Exe</OutputType>
+    <AppDesignerFolder>Properties</AppDesignerFolder>
+    <RootNamespace>HTML5LaunchHelper</RootNamespace>
+    <AssemblyName>HTML5LaunchHelper</AssemblyName>
+    <TargetFrameworkVersion>v4.5</TargetFrameworkVersion>
+    <FileAlignment>512</FileAlignment>
+  </PropertyGroup>
+  <PropertyGroup Condition=" '$(Configuration)|$(Platform)' == 'Debug|AnyCPU' ">
+    <PlatformTarget>AnyCPU</PlatformTarget>
+    <DebugSymbols>true</DebugSymbols>
+    <DebugType>full</DebugType>
+    <Optimize>false</Optimize>
+    <OutputPath>..\..\..\..\Binaries\DotNET\</OutputPath>
+    <DefineConstants>DEBUG;TRACE</DefineConstants>
+    <ErrorReport>prompt</ErrorReport>
+    <WarningLevel>4</WarningLevel>
+  </PropertyGroup>
+  <PropertyGroup Condition=" '$(Configuration)|$(Platform)' == 'Development|AnyCPU' ">
+    <PlatformTarget>AnyCPU</PlatformTarget>
+    <DebugType>pdbonly</DebugType>
+    <Optimize>true</Optimize>
+    <OutputPath>..\..\..\..\Binaries\DotNET\</OutputPath>
+    <DefineConstants>TRACE</DefineConstants>
+    <ErrorReport>prompt</ErrorReport>
+    <WarningLevel>4</WarningLevel>
+  </PropertyGroup>
+  <ItemGroup>
+    <Reference Include="System" />
+    <Reference Include="System.Core" />
+    <Reference Include="System.Xml.Linq" />
+    <Reference Include="System.Data.DataSetExtensions" />
+    <Reference Include="Microsoft.CSharp" />
+    <Reference Include="System.Data" />
+    <Reference Include="System.Xml" />
+  </ItemGroup>
+  <ItemGroup>
+    <Compile Include="Program.cs" />
+    <Compile Include="Properties\AssemblyInfo.cs" />
+  </ItemGroup>
+  <ItemGroup>
+    <None Include="App.config" />
+  </ItemGroup>
+  <Import Project="$(MSBuildToolsPath)\Microsoft.CSharp.targets" />
+  <!-- To modify your build process, add your task inside one of the targets below and uncomment it. 
+       Other similar extension points exist, see Microsoft.Common.targets.
+  <Target Name="BeforeBuild">
+  </Target>
+  <Target Name="AfterBuild">
+  </Target>
+  -->
+  <ItemGroup>
+    <ProjectReference Include="..\..\UnrealBuildTool\UnrealBuildTool.csproj">
+      <Project>{fd7c5e1a-cfe4-4fd5-a525-1eb1599a39ac}</Project>
+      <Name>UnrealBuildTool</Name>
+      <Private>False</Private>
+    </ProjectReference>
+  </ItemGroup>
+</Project>
\ No newline at end of file
diff --git a/Engine/Source/Programs/HTML5/HTML5LaunchHelper/Program.cs b/Engine/Source/Programs/HTML5/HTML5LaunchHelper/Program.cs
new file mode 100644
--- /dev/null
+++ b/Engine/Source/Programs/HTML5/HTML5LaunchHelper/Program.cs
@@ -0,0 +1,295 @@
+﻿using System;
+using System.Collections.Generic;
+using System.Linq;
+using System.Text;
+using System.Threading.Tasks;
+using System.Diagnostics;
+using UnrealBuildTool;
+
+// Example Command Lines
+// HTML5LaunchHelper -m -h -b /Applications/Safari.app localhost:8000
+// HTML5LaunchHelper -m -h -b /Applications/Firefox.app http://localhost:8000 -silent -profile [REDACTED_USER_HOME]/webporfiles/nightly/UE4HTML5 --jsconsole
+// HTML5LaunchHelper -m -h -b /Applications/Google\ Chrome.app http://localhost:8000 --user-data-dir="[REDACTED_USER_HOME]/webporfiles/chrome/UE4HTML5" --enable-logging --no-first-run
+namespace HTML5LaunchHelper
+{
+	class Program
+	{
+		static private bool 	bMultiprocessBrowser = false;
+		static private bool 	bStartPythonServer = false;
+		static private int  	PythonServerPort = 8000;
+		static private string 	PythonPath = "python";
+		static private string 	BrowserPath = null;
+		static private string	ServerDirectory = null;
+		static private string 	Url = null;
+		static private string   AdditionalArgs = "";
+		static private List<Process> ProcessesToKill = new List<Process>();
+		static private List<Process> ProcessesToWatch = new List<Process>();
+
+		static Process SpawnBrowserProcess(string bpath, string args)
+		{
+			var Result = new Process();
+			if (BuildHostPlatform.Current.Platform == UnrealTargetPlatform.Mac)
+			{
+	            Result.StartInfo.FileName = "/usr/bin/open";
+				Result.StartInfo.UseShellExecute = false;
+				Result.StartInfo.RedirectStandardOutput = true;
+				Result.StartInfo.RedirectStandardInput = true;
+				Result.StartInfo.Arguments = String.Format("-nW \"{0}\" --args {1}", bpath, args);
+				Result.EnableRaisingEvents = true;
+			}
+			else
+			{
+	            Result.StartInfo.FileName = bpath;
+				Result.StartInfo.UseShellExecute = false;
+				Result.StartInfo.RedirectStandardOutput = true;
+				Result.StartInfo.RedirectStandardInput = true;
+				Result.StartInfo.Arguments = args;
+				Result.EnableRaisingEvents = true;
+			}
+			System.Console.WriteLine("Spawning Browser Process {0} with args {1}\n", bpath, args);
+			return Result;
+		}
+
+		static int Main(string[] args)
+		{
+			if (!ParseArgs(args))
+			{
+				// TODO: print error & options.
+				System.Console.Write("Invalid args\n");
+				return -1;
+			}
+
+			if (bStartPythonServer)
+			{
+				BeginPythonServer();
+				// Give the server time to start. 1 sec should be enough
+				System.Threading.Thread.Sleep(1000);
+			}
+
+			var bIsSafari = BrowserPath.Contains("Safari");
+
+			// Browsers can be multiprocess programs (Chrome, basically)
+			// which we need to catch spawning of other child processes. The trick is
+			// they aren't really child-processes. There appears to be no real binding between the two
+			// so we kinda fudge it a bit here
+			var PrevProcesses = Process.GetProcesses();
+			var FirstProcess = SpawnBrowserProcess(BrowserPath, !bIsSafari ? Url + " " + AdditionalArgs : "");
+			ProcessesToWatch.Add(FirstProcess);
+			FirstProcess.Start();
+			var ProcName = FirstProcess.ProcessName;
+
+			if (bIsSafari)
+			{
+				// Give Safari time to open...
+				System.Threading.Thread.Sleep(2000);
+				var Proc = new Process();
+				Proc.StartInfo.FileName = "/usr/bin/osascript";
+				Proc.StartInfo.UseShellExecute = false;
+				Proc.StartInfo.RedirectStandardOutput = true;
+				Proc.StartInfo.RedirectStandardInput = true;
+				Proc.StartInfo.Arguments = String.Format("-e 'tell application \"Safari\" to open location \"{1}\"'", FirstProcess.Id, Url);
+				Proc.EnableRaisingEvents = true;
+				Proc.Start();
+				Proc.WaitForExit();	
+			}
+
+			// We should now have a list of processes to watch to exit.
+			// Loop over the calling WaitForExit() until the list is empty.
+			while (ProcessesToWatch.Count() > 0)
+			{
+				for (var i=0; i<ProcessesToWatch.Count(); ++i)
+				{
+					if (ProcessesToWatch[i].HasExited)
+					{
+						ProcessesToWatch.RemoveAt(i);
+					}
+				}
+
+				if (bMultiprocessBrowser && FirstProcess != null && FirstProcess.HasExited)
+				{
+					var CurrentProcesses = Process.GetProcesses();
+					foreach (var item in CurrentProcesses)
+					{
+						var bWasAlive = false;
+						foreach (var pitem in PrevProcesses)
+						{
+							if (pitem.Id == item.Id)
+							{
+								bWasAlive = true;
+							}
+						}
+						if (!bWasAlive)
+						{
+							try {
+								if (!item.HasExited && item.ProcessName.StartsWith(ProcName))
+								{
+									var PID = item.Id;
+									System.Console.WriteLine("Found Process {0} with PID {1} which started at {2}. Waiting on that process to end.", item.ProcessName, item.Id, item.StartTime.ToString());
+									ProcessesToWatch.Add(item);
+								}
+							} catch {}
+						}
+					}
+					FirstProcess = null;
+				}
+
+				//If any service processes have died this is considered an error
+				foreach (var proc in ProcessesToKill)
+				{
+					if (proc.HasExited)
+					{
+						System.Console.WriteLine("A spawned thread has died. Do you have a python server instance running?");
+						HardShutdown();
+						return -1;
+					}
+				}
+
+				System.Threading.Thread.Sleep(250);
+			}
+			System.Console.WriteLine("All processes being watched have exited.\n");
+
+			//All processes we cared about have finished so clean up the services we spawned for them.
+			foreach (var proc in ProcessesToKill)
+			{
+				if (!proc.HasExited)
+				{
+					System.Console.WriteLine("Killing spawned process {0}.\n", proc.Id);
+					proc.Kill();
+					proc.WaitForExit();
+				}
+			}
+
+			return 0;
+		}
+
+		static void HardShutdown()
+		{
+			foreach (var proc in ProcessesToKill)
+			{
+				if (!proc.HasExited)
+				{
+					proc.Kill();
+					proc.WaitForExit();
+				}
+			}
+
+			foreach (var proc in ProcessesToWatch)
+			{
+				if (!proc.HasExited)
+				{
+					proc.Kill();
+					proc.WaitForExit();
+				}
+			}
+		}
+
+		static void BeginPythonServer()
+		{
+			var Result = new Process();
+            Result.StartInfo.FileName = PythonPath;
+			Result.StartInfo.UseShellExecute = false;
+			Result.StartInfo.RedirectStandardOutput = true;
+			Result.StartInfo.RedirectStandardInput = true;
+			Result.StartInfo.Arguments = String.Format("-m SimpleHTTPServer {0}", PythonServerPort);
+			Result.StartInfo.WorkingDirectory = ServerDirectory;
+			Result.EnableRaisingEvents = true;
+			Result.Start();
+			System.Console.Write("Started New Process (PID={0} Name={1})\n", Result.Id, Result.ProcessName);
+			ProcessesToKill.Add(Result);
+		}
+
+		static bool ParseArgs(string[] args)
+		{
+			bool bOptionalPythonName = false;
+			bool bParseBrowserPath = false;
+			bool bParseServerDir = false;
+			bool bParseServerPort = false;
+
+			foreach (var arg in args)
+			{
+				var a = arg.Trim();
+				if (bOptionalPythonName || bParseBrowserPath || bParseServerDir || bParseServerPort)
+				{
+					if (!a.StartsWith("-"))
+					{
+						if (bOptionalPythonName)
+						{
+							PythonPath = a;
+						}
+						else if (bParseBrowserPath)
+						{
+							BrowserPath = a;
+						}
+						else if (bParseServerDir)
+						{
+							ServerDirectory = a;
+						}
+						else if (bParseServerPort)
+						{
+							PythonServerPort = Convert.ToInt32(a);
+						}
+					}
+					bOptionalPythonName = bParseBrowserPath = bParseServerDir = bParseServerPort = false;
+				}
+				else if (a == "-m")
+				{
+					bMultiprocessBrowser = true;
+				}
+				else if (a == "-p")
+				{
+					bOptionalPythonName = true;
+				}
+				else if (a == "-h")
+				{
+					bStartPythonServer = true;
+				}
+				else if (a == "-b")
+				{
+					bParseBrowserPath = true;
+				}
+				else if (a == "-w")
+				{
+					bParseServerDir = true;
+				}
+				else if (a == "-s")
+				{
+					bParseServerPort = true;
+				}
+				else if (string.IsNullOrEmpty(Url))
+				{
+					Url = a;
+				}
+				else
+				{
+					if (a.StartsWith ("-") && a.Contains('=') && a.Contains(' ')) 
+					{
+						var Parts = a.Split('=');
+						a = "";
+						for (var p = 0; p < Parts.Count(); ++p)
+						{
+							if (Parts [p].Contains(' '))
+							{
+								a += "\""+Parts[p]+"\"";
+							}
+							else
+							{
+								a += Parts[p];
+							}
+							if (p < Parts.Count() - 1)
+							{
+								a += '=';
+							}
+						}
+					}
+					else if (a.Contains(' '))
+					{
+						a = "\""+a+"\"";
+					}
+					AdditionalArgs += a + " ";
+				}
+			}
+
+			return !string.IsNullOrEmpty(BrowserPath) && !string.IsNullOrEmpty(Url) && !string.IsNullOrEmpty(ServerDirectory);
+		}
+	}
+}
diff --git a/Engine/Source/Programs/HTML5/HTML5LaunchHelper/Properties/AssemblyInfo.cs b/Engine/Source/Programs/HTML5/HTML5LaunchHelper/Properties/AssemblyInfo.cs
new file mode 100644
--- /dev/null
+++ b/Engine/Source/Programs/HTML5/HTML5LaunchHelper/Properties/AssemblyInfo.cs
@@ -0,0 +1,36 @@
+﻿using System.Reflection;
+using System.Runtime.CompilerServices;
+using System.Runtime.InteropServices;
+
+// General Information about an assembly is controlled through the following 
+// set of attributes. Change these attribute values to modify the information
+// associated with an assembly.
+[assembly: AssemblyTitle("HTML5LaunchHelper")]
+[assembly: AssemblyDescription("")]
+[assembly: AssemblyConfiguration("")]
+[assembly: AssemblyCompany("")]
+[assembly: AssemblyProduct("HTML5LaunchHelper")]
+[assembly: AssemblyCopyright("Copyright ©  2015")]
+[assembly: AssemblyTrademark("")]
+[assembly: AssemblyCulture("")]
+
+// Setting ComVisible to false makes the types in this assembly not visible 
+// to COM components.  If you need to access a type in this assembly from 
+// COM, set the ComVisible attribute to true on that type.
+[assembly: ComVisible(false)]
+
+// The following GUID is for the ID of the typelib if this project is exposed to COM
+[assembly: Guid("c22fe09b-9d62-4b4c-ae2e-bdb5d66c143f")]
+
+// Version information for an assembly consists of the following four values:
+//
+//      Major Version
+//      Minor Version 
+//      Build Number
+//      Revision
+//
+// You can specify all the values or you can default the Build and Revision Numbers 
+// by using the '*' as shown below:
+// [assembly: AssemblyVersion("1.0.*")]
+[assembly: AssemblyVersion("1.0.0.0")]
+[assembly: AssemblyFileVersion("1.0.0.0")]
diff --git a/Engine/Source/Programs/UnrealBuildTool/Configuration/EngineConfiguration.cs b/Engine/Source/Programs/UnrealBuildTool/Configuration/EngineConfiguration.cs
--- a/Engine/Source/Programs/UnrealBuildTool/Configuration/EngineConfiguration.cs
+++ b/Engine/Source/Programs/UnrealBuildTool/Configuration/EngineConfiguration.cs
@@ -545,21 +545,18 @@ namespace UnrealBuildTool
 				}
 			}
 
-			string UserSettingsFolder = null; // Match FPlatformProcess::UserSettingsDir()
+			string UserSettingsFolder = Utils.GetUserSettingDirectory(); // Match FPlatformProcess::UserSettingsDir()
 			string PersonalFolder = null; // Match FPlatformProcess::UserDir()
 			if (BuildHostPlatform.Current.Platform == UnrealTargetPlatform.Mac)
 			{
-				UserSettingsFolder = Path.Combine(Environment.GetFolderPath(Environment.SpecialFolder.Personal), "Library", "Application Support", "Epic");
 				PersonalFolder = Path.Combine(Environment.GetFolderPath(Environment.SpecialFolder.Personal), "Documents");
 			}
 			else if (Environment.OSVersion.Platform == PlatformID.Unix)
 			{
-				UserSettingsFolder = Path.Combine(Environment.GetFolderPath(Environment.SpecialFolder.ApplicationData), "Epic");
 				PersonalFolder = Path.Combine(Environment.GetFolderPath(Environment.SpecialFolder.Personal), "Documents");
 			}
 			else
 			{
-				UserSettingsFolder = Environment.GetFolderPath(Environment.SpecialFolder.LocalApplicationData);
 				PersonalFolder = Environment.GetFolderPath(Environment.SpecialFolder.Personal);
 			}
 
diff --git a/Engine/Source/Programs/UnrealBuildTool/HTML5/HTML5SDKInfo.cs b/Engine/Source/Programs/UnrealBuildTool/HTML5/HTML5SDKInfo.cs
--- a/Engine/Source/Programs/UnrealBuildTool/HTML5/HTML5SDKInfo.cs
+++ b/Engine/Source/Programs/UnrealBuildTool/HTML5/HTML5SDKInfo.cs
@@ -110,6 +110,73 @@ namespace UnrealBuildTool
 			return false;
 		}
 
+		private class SDKVersion : IComparable<SDKVersion>
+		{
+			public System.Version Version;
+
+			public string Directory;
+
+			public static bool operator ==(SDKVersion LHS, SDKVersion RHS)
+			{
+				return LHS.Version == RHS.Version;
+			}
+
+			public static bool operator !=(SDKVersion LHS, SDKVersion RHS)
+			{
+				return LHS.Version != RHS.Version;
+			}
+
+			// Override the Object.Equals(object o) method:
+			public override bool Equals(object o)
+			{
+				try
+				{
+					return (bool)(this == (SDKVersion)o);
+				}
+				catch
+				{
+					return false;
+				}
+			}
+
+			// Override the Object.GetHashCode() method:
+			public override int GetHashCode()
+			{
+				return Version.GetHashCode();
+			}
+
+			public int CompareTo(SDKVersion other)
+			{
+				return this.Version.CompareTo(other.Version);
+			}
+		}
+
+		static List<SDKVersion> GetInstalledVersions(string RootPath)
+		{
+			List<SDKVersion> Versions = new List<SDKVersion>();
+			SDKVersion Ver = null;
+			if (!System.IO.Directory.Exists(Path.Combine(RootPath, "emscripten")))
+			{
+				return Versions;
+			}
+
+			string[] Directories = Directory.GetDirectories(Path.Combine(RootPath, "emscripten"), "*", SearchOption.TopDirectoryOnly);
+			foreach (var Dir in Directories)
+			{
+				var VersionFilePath = Path.Combine(Dir, "emscripten-version.txt");
+				if (System.IO.File.Exists(VersionFilePath))
+				{
+					Ver = new SDKVersion();
+					Ver.Version = System.Version.Parse(ReadVersionFile(VersionFilePath));
+					Ver.Directory = Dir;
+					Versions.Add(Ver);
+				}
+			}
+
+			Versions.Sort();
+			return Versions;
+		}
+
 		static void EnsureConfigCacheIsReady()
 		{
 			if (ConfigCache == null)
@@ -132,9 +199,57 @@ namespace UnrealBuildTool
         {
 			EnsureConfigCacheIsReady();
 			// Search the ini file for Emscripten root first
+			string EmscriptenRoot;
 			string EMCCPath;
-			if (ConfigCache.GetPath("/Script/HTML5PlatformEditor.HTML5SDKSettings", "Emscripten", out EMCCPath) && System.IO.Directory.Exists(EMCCPath))
- 				return EMCCPath;
+			if (ConfigCache.GetString("/Script/HTML5PlatformEditor.HTML5SDKSettings", "EmscriptenRoot", out EmscriptenRoot))
+			{
+				string SDKPathString = null;
+				string VersionString = null;
+				var Index = EmscriptenRoot.IndexOf("SDKPath=");
+				if (Index != -1)
+				{
+					Index+=9;
+					var Index2 = EmscriptenRoot.IndexOf("\"", Index);
+					if (Index2 != -1)
+					{
+						SDKPathString = EmscriptenRoot.Substring(Index, Index2 - Index);
+					}
+				}
+				Index = EmscriptenRoot.IndexOf("EmscriptenVersion=");
+				if (Index != -1)
+				{
+					Index += 19;
+					var Index2 = EmscriptenRoot.IndexOf("\"", Index);
+					if (Index2 != -1)
+					{
+						VersionString = EmscriptenRoot.Substring(Index, Index2 - Index);
+					}
+				}
+
+				if (!string.IsNullOrEmpty(SDKPathString) && !string.IsNullOrEmpty(VersionString))
+				{
+					var SDKVersions = GetInstalledVersions(SDKPathString);
+
+					if (VersionString == "-1.-1.-1" && SDKVersions.Count > 0)
+					{
+						var Ver = SDKVersions[SDKVersions.Count - 1];
+						if (System.IO.Directory.Exists(Ver.Directory))
+						{
+							return Ver.Directory;
+						}
+					}
+
+					var RequiredVersion = System.Version.Parse(VersionString);
+					foreach (var i in SDKVersions)
+					{
+ 						if (i.Version == RequiredVersion && System.IO.Directory.Exists(i.Directory))
+						{
+							return i.Directory;
+						}
+					}
+				}
+			}
+
 
 			// Old Method
 			if (bAllowFallbackSDKSettings)
@@ -280,11 +395,16 @@ namespace UnrealBuildTool
 			return true;
 		}
 
+		public static string ReadVersionFile(string path)
+		{
+			string VersionInfo = File.ReadAllText(path);
+			VersionInfo = VersionInfo.Trim();
+			return VersionInfo; 
+		}
+
         public static string EmscriptenVersion()
         {
-            string VersionInfo = File.ReadAllText(GetVersionInfoPath());
-            VersionInfo = VersionInfo.Trim();
-            return VersionInfo; 
+			return ReadVersionFile(GetVersionInfoPath());
         }
 
 		static string GetVersionInfoPath()
diff --git a/Engine/Source/Programs/UnrealBuildTool/HTML5/HTML5ToolChain.cs b/Engine/Source/Programs/UnrealBuildTool/HTML5/HTML5ToolChain.cs
--- a/Engine/Source/Programs/UnrealBuildTool/HTML5/HTML5ToolChain.cs
+++ b/Engine/Source/Programs/UnrealBuildTool/HTML5/HTML5ToolChain.cs
@@ -13,6 +13,8 @@ namespace UnrealBuildTool
 	{
         static string EMCCPath;
         static string PythonPath;
+        // Debug options -- TODO: add these to SDK/Project setup?
+        static bool   bEnableTracing = false;
 
 		// cache the location of SDK tools
 		public override void RegisterToolChain()
@@ -82,6 +84,11 @@ namespace UnrealBuildTool
                 Result += " -s NO_EXIT_RUNTIME=1 --memory-init-file 1";
             }
 
+            if (bEnableTracing)
+            {
+            	Result += " --tracing";
+            }
+
             return Result;
 		}
 		
@@ -310,6 +317,11 @@ namespace UnrealBuildTool
 				Arguments += string.Format(" -D{0}", Definition);
 			}
 
+			if (bEnableTracing)
+			{
+				Arguments += string.Format(" -D__EMSCRIPTEN_TRACING__");
+			}
+
         
 			var BuildPlatform = UEBuildPlatform.GetBuildPlatformForCPPTargetPlatform(CompileEnvironment.Config.Target.Platform);
 
@@ -349,7 +361,7 @@ namespace UnrealBuildTool
 
 				CompileAction.CommandArguments = EMCCPath + " " + Arguments + FileArguments + CompileEnvironment.Config.AdditionalArguments;
 
-                System.Console.WriteLine(CompileAction.CommandArguments); 
+                //System.Console.WriteLine(CompileAction.CommandArguments); 
 				CompileAction.StatusDescription = Path.GetFileName(SourceFile.AbsolutePath);
 				CompileAction.OutputEventHandler = new DataReceivedEventHandler(CompileOutputReceivedDataEventHandler);
 
@@ -459,7 +471,7 @@ namespace UnrealBuildTool
 			// Add the input files to a response file, and pass the response file on the command-line.
 			foreach (FileItem InputFile in LinkEnvironment.InputFiles)
 			{
-                System.Console.WriteLine("File  {0} ", InputFile.AbsolutePath);
+                //System.Console.WriteLine("File  {0} ", InputFile.AbsolutePath);
                 LinkAction.CommandArguments += string.Format(" \"{0}\"", InputFile.AbsolutePath);
 				LinkAction.PrerequisiteItems.Add(InputFile);
 			}
diff --git a/Engine/Source/Programs/UnrealBuildTool/System/ProjectFileGenerator.cs b/Engine/Source/Programs/UnrealBuildTool/System/ProjectFileGenerator.cs
--- a/Engine/Source/Programs/UnrealBuildTool/System/ProjectFileGenerator.cs
+++ b/Engine/Source/Programs/UnrealBuildTool/System/ProjectFileGenerator.cs
@@ -515,6 +515,8 @@ namespace UnrealBuildTool
 
                     // Add all of the PS4 C# projects
                     AddPS4Projects( ProgramsFolder );
+
+					AddHTML5Projects( ProgramsFolder );
                 }
 
 
@@ -1278,6 +1280,19 @@ namespace UnrealBuildTool
 		private void AddAndroidProjects(MasterProjectFolder Folder)
 		{
 		}
+
+		/// <summary>
+		/// Adds all of the HTML5 C# projects to the master project
+		/// </summary>
+		private void AddHTML5Projects(MasterProjectFolder Folder)
+		{
+			string ProjectFolderName = Path.Combine(EngineRelativePath, "Source", "Programs", "HTML5");
+			DirectoryInfo ProjectFolderInfo = new DirectoryInfo(ProjectFolderName);
+			if (ProjectFolderInfo.Exists)
+			{
+				Folder.ChildProjects.Add(AddSimpleCSharpProject("HTML5/HTML5LaunchHelper", true)); // Build by default; needed for UAT.
+			}
+		}
 		
 
 		/// <summary>
diff --git a/Engine/Source/Programs/UnrealBuildTool/System/Utils.cs b/Engine/Source/Programs/UnrealBuildTool/System/Utils.cs
--- a/Engine/Source/Programs/UnrealBuildTool/System/Utils.cs
+++ b/Engine/Source/Programs/UnrealBuildTool/System/Utils.cs
@@ -859,6 +859,25 @@ namespace UnrealBuildTool
 		{
 			return Array.IndexOf(TestType.GetInterfaces(), typeof(InterfaceType)) != -1;
 		}
+
+		/// <summary>
+		/// Returns the User Settings Directory path. This matches FPlatformProcess::UserSettingsDir()
+		/// </summary>
+		public static string GetUserSettingDirectory()
+		{
+			if (BuildHostPlatform.Current.Platform == UnrealTargetPlatform.Mac)
+			{
+				return Path.Combine(Environment.GetFolderPath(Environment.SpecialFolder.Personal), "Library", "Application Support", "Epic");
+			}
+			else if (Environment.OSVersion.Platform == PlatformID.Unix)
+			{
+				return Path.Combine(Environment.GetFolderPath(Environment.SpecialFolder.ApplicationData), "Epic");
+			}
+			else
+			{
+				return Environment.GetFolderPath(Environment.SpecialFolder.LocalApplicationData);
+			}
+		}
 	}
 
 	/// <summary>
diff --git a/Engine/Source/Programs/UnrealPak/Private/UnrealPak.cpp b/Engine/Source/Programs/UnrealPak/Private/UnrealPak.cpp
--- a/Engine/Source/Programs/UnrealPak/Private/UnrealPak.cpp
+++ b/Engine/Source/Programs/UnrealPak/Private/UnrealPak.cpp
@@ -767,10 +767,21 @@ bool CreatePakFile(const TCHAR* Filename, TArray<FPakInputPair>& FilesToAdd, con
 		{
 			if (CompressedFileBuffer.CompressFileToWorkingBuffer(FilesToAdd[FileIndex], ReadBuffer, BufferSize, CompressionMethod, CmdLineParameters.CompressionBlockSize))
 			{
-				NewEntry.Info.CompressionMethod = CompressionMethod;
-				NewEntry.Info.CompressionBlocks.AddUninitialized(CompressedFileBuffer.CompressedBlocks.Num());
-				RealFileSize = CompressedFileBuffer.TotalCompressedSize + NewEntry.Info.GetSerializedSize(FPakInfo::PakFile_Version_Latest);
-				NewEntry.Info.CompressionBlocks.Reset();
+				// Check the compression ratio, if it's too low just store uncompressed. Also take into account read size
+				// if we still save 64KB it's probably worthwhile compressing, as that saves a file read operation in the runtime.
+				// TODO: drive this threashold from the command line
+				float PercentLess = ((float)CompressedFileBuffer.TotalCompressedSize/(OriginalFileSize/100.f));
+				if (PercentLess > 90.f && (OriginalFileSize-CompressedFileBuffer.TotalCompressedSize) < 65536)
+				{
+					CompressionMethod = COMPRESS_None;
+				}
+				else
+				{
+					NewEntry.Info.CompressionMethod = CompressionMethod;
+					NewEntry.Info.CompressionBlocks.AddUninitialized(CompressedFileBuffer.CompressedBlocks.Num());
+					RealFileSize = CompressedFileBuffer.TotalCompressedSize + NewEntry.Info.GetSerializedSize(FPakInfo::PakFile_Version_Latest);
+					NewEntry.Info.CompressionBlocks.Reset();
+				}
 			}
 		}
 
@@ -800,9 +811,10 @@ bool CreatePakFile(const TCHAR* Filename, TArray<FPakInputPair>& FilesToAdd, con
 			// Update offset now and store it in the index (and only in index)
 			NewEntry.Info.Offset = NewEntryOffset;
 			Index.Add(NewEntry);
-			if (FilesToAdd[FileIndex].bNeedsCompression)
+			if (FilesToAdd[FileIndex].bNeedsCompression && CompressionMethod != COMPRESS_None)
 			{
-				UE_LOG(LogPakFile, Display, TEXT("Added compressed file \"%s\", Compressed Size %lld bytes, Original Size %lld bytes."), *NewEntry.Filename, NewEntry.Info.Size, NewEntry.Info.UncompressedSize);
+				float PercentLess = ((float)NewEntry.Info.Size/(NewEntry.Info.UncompressedSize/100.f));
+				UE_LOG(LogPakFile, Display, TEXT("Added compressed file \"%s\", %.2f%% of original size. Compressed Size %lld bytes, Original Size %lld bytes. "), *NewEntry.Filename, PercentLess, NewEntry.Info.Size, NewEntry.Info.UncompressedSize);
 			}
 			else
 			{
diff --git a/Engine/Source/Runtime/Core/Private/HTML5/HTML5Application.cpp b/Engine/Source/Runtime/Core/Private/HTML5/HTML5Application.cpp
--- a/Engine/Source/Runtime/Core/Private/HTML5/HTML5Application.cpp
+++ b/Engine/Source/Runtime/Core/Private/HTML5/HTML5Application.cpp
@@ -5,6 +5,9 @@
 #include "HTML5Cursor.h"
 #include "HTML5InputInterface.h"
 
+#include "SDL_opengl.h"
+
+DEFINE_LOG_CATEGORY_STATIC(LogHTML5Application, Log, All);
 
 #if PLATFORM_HTML5_BROWSER
 #include "emscripten.h"
@@ -63,23 +66,23 @@ void FHTML5Application::PollGameDeviceState( const float TimeDelta )
 
 
 					// ignore resized client Height/Width
-					int Width;
-					int Height;
 #if PLATFORM_HTML5_BROWSER
 					int fs;
-					emscripten_get_canvas_size(&Width, &Height, &fs);
+					emscripten_get_canvas_size(&WindowWidth, &WindowHeight, &fs);
+					UE_LOG(LogHTML5Application, Verbose, TEXT("emscripten_get_canvas_size: Width:%d, Height:%d, Fullscreen:%d"), WindowWidth, WindowHeight, fs);
 #endif 
 
 #if PLATFORM_HTML5_WIN32 
-					Width = windowEvent.data1; 
-					Height = windowEvent.data2; 
+					WindowWidth = windowEvent.data1; 
+					WindowHeight = windowEvent.data2; 
 #endif 
 
 					switch (windowEvent.event)
 					{
 					case SDL_WINDOWEVENT_SIZE_CHANGED:
 						{
-							MessageHandler->OnSizeChanged(ApplicationWindow,Width,Height, false);
+							UE_LOG(LogHTML5Application, Verbose, TEXT("WindowSizeChanged: Width:%d, Height:%d"), WindowWidth, WindowHeight);
+							MessageHandler->OnSizeChanged(ApplicationWindow,WindowWidth,WindowHeight, false);
 							MessageHandler->OnResizingWindow(ApplicationWindow);
 
 							FDisplayMetrics DisplayMetrics;
@@ -89,6 +92,7 @@ void FHTML5Application::PollGameDeviceState( const float TimeDelta )
 						break;
 					case SDL_WINDOWEVENT_RESIZED:
 						{
+							UE_LOG(LogHTML5Application, Verbose, TEXT("WindowResized: Width:%d, Height:%d"), WindowWidth, WindowHeight);
 							MessageHandler->OnResizingWindow(ApplicationWindow);
 
 							FDisplayMetrics DisplayMetrics;
@@ -98,6 +102,7 @@ void FHTML5Application::PollGameDeviceState( const float TimeDelta )
 						break;
 					case SDL_WINDOWEVENT_ENTER:
 						{
+							UE_LOG(LogHTML5Application, Verbose, TEXT("WindowEnter"));
 							MessageHandler->OnCursorSet();
 							MessageHandler->OnWindowActivationChanged(ApplicationWindow, EWindowActivation::Activate); 
 							WarmUpTicks = 0; 
@@ -105,17 +110,20 @@ void FHTML5Application::PollGameDeviceState( const float TimeDelta )
 						break;
 					case SDL_WINDOWEVENT_LEAVE:
 						{
+							UE_LOG(LogHTML5Application, Verbose, TEXT("WindowLeave"));
 							MessageHandler->OnWindowActivationChanged(ApplicationWindow, EWindowActivation::Deactivate);
 						}
 						break;
 					case SDL_WINDOWEVENT_FOCUS_GAINED:
 						{
+							UE_LOG(LogHTML5Application, Verbose, TEXT("WindowFocusGained"));
 							MessageHandler->OnWindowActivationChanged(ApplicationWindow, EWindowActivation::Activate);
 									WarmUpTicks = 0;
 						}
 						break;
 					case SDL_WINDOWEVENT_FOCUS_LOST:
 						{
+							UE_LOG(LogHTML5Application, Verbose, TEXT("WindowFocusLost"));
 							MessageHandler->OnWindowActivationChanged(ApplicationWindow, EWindowActivation::Deactivate);
 						}
 						break;
@@ -161,10 +169,12 @@ void FHTML5Application::PollGameDeviceState( const float TimeDelta )
 #if PLATFORM_HTML5_BROWSER
 		if (((FHTML5Cursor*)Cursor.Get())->LockStatus && !((FHTML5Cursor*)Cursor.Get())->CursorStatus)
 		{
+			UE_LOG(LogHTML5Application, Verbose, TEXT("Request pointer lock"));
 			emscripten_request_pointerlock ( "#canvas" , true);
 		}
 		else
 		{
+			UE_LOG(LogHTML5Application, Verbose, TEXT("Exit pointer lock"));
 			emscripten_exit_pointerlock(); 
 		}
 #endif 
@@ -184,6 +194,7 @@ void FDisplayMetrics::GetDisplayMetrics(FDisplayMetrics& OutDisplayMetrics)
 	OutDisplayMetrics.VirtualDisplayRect    =	OutDisplayMetrics.PrimaryDisplayWorkAreaRect;
 	OutDisplayMetrics.PrimaryDisplayWidth   =	OutDisplayMetrics.PrimaryDisplayWorkAreaRect.Right;
 	OutDisplayMetrics.PrimaryDisplayHeight  =	OutDisplayMetrics.PrimaryDisplayWorkAreaRect.Bottom; 
+	UE_LOG(LogHTML5Application, Verbose, TEXT("GetDisplayMetrics Width:%d, Height:%d"), OutDisplayMetrics.PrimaryDisplayWorkAreaRect.Right, OutDisplayMetrics.PrimaryDisplayWorkAreaRect.Bottom);
 }
 
 TSharedRef< FGenericWindow > FHTML5Application::MakeWindow()
diff --git a/Engine/Source/Runtime/Core/Private/HTML5/HTML5InputInterface.cpp b/Engine/Source/Runtime/Core/Private/HTML5/HTML5InputInterface.cpp
--- a/Engine/Source/Runtime/Core/Private/HTML5/HTML5InputInterface.cpp
+++ b/Engine/Source/Runtime/Core/Private/HTML5/HTML5InputInterface.cpp
@@ -7,6 +7,7 @@
 
 DECLARE_LOG_CATEGORY_EXTERN(LogHTML5, Log, All);
 DEFINE_LOG_CATEGORY(LogHTML5);
+DEFINE_LOG_CATEGORY_STATIC(LogHTML5Input, Log, All);
 
 
 #if PLATFORM_HTML5_BROWSER
@@ -186,7 +187,6 @@ static bool GeneratesKeyCharMessage(const SDL_KeyboardEvent & KeyDownEvent)
 		(Sym != SDLK_DOWN && Sym != SDLK_LEFT && Sym != SDLK_RIGHT && Sym != SDLK_UP && Sym != SDLK_DELETE);
 }
 
-
 void FHTML5InputInterface::Tick(float DeltaTime, const SDL_Event& Event,TSharedRef < FGenericWindow>& ApplicationWindow )
 {
 
@@ -206,6 +206,7 @@ void FHTML5InputInterface::Tick(float DeltaTime, const SDL_Event& Event,TSharedR
 					const TCHAR Character = ConvertChar(KeyEvent.keysym);
 					MessageHandler->OnKeyChar(Character, bIsRepeated);
 				}
+				UE_LOG(LogHTML5Input, Verbose, TEXT("KeyDown: Code:%d bIsRepeated:%s"), KeyCode, bIsRepeated ? TEXT("TRUE") : TEXT("FALSE"));
 			}
 			break;
 		case SDL_KEYUP:
@@ -216,18 +217,21 @@ void FHTML5InputInterface::Tick(float DeltaTime, const SDL_Event& Event,TSharedR
 				const bool IsRepeat = keyEvent.repeat != 0;
 
 				MessageHandler->OnKeyUp( KeyCode, keyEvent.keysym.sym, IsRepeat );
+				UE_LOG(LogHTML5Input, Verbose, TEXT("KeyUp Code:%d"), KeyCode);
 			}
 			break;
 		case SDL_MOUSEBUTTONDOWN:
 			{
 				EMouseButtons::Type MouseButton = Event.button.button == 1 ? EMouseButtons::Left : Event.button.button == 2 ? EMouseButtons::Middle : EMouseButtons::Right;
 				MessageHandler->OnMouseDown(ApplicationWindow, MouseButton );
+				UE_LOG(LogHTML5Input, Verbose, TEXT("MouseButtonDown ID:%d"), Event.button.button);
 			}
 			break;
 		case SDL_MOUSEBUTTONUP:
 			{
 				EMouseButtons::Type MouseButton = Event.button.button == 1 ? EMouseButtons::Left : Event.button.button == 2 ? EMouseButtons::Middle : EMouseButtons::Right;
 				MessageHandler->OnMouseUp(MouseButton );
+				UE_LOG(LogHTML5Input, Verbose, TEXT("MouseButtonUp ID:%d"), Event.button.button);
 			}
 			break; 
 		case SDL_MOUSEMOTION:
@@ -235,6 +239,7 @@ void FHTML5InputInterface::Tick(float DeltaTime, const SDL_Event& Event,TSharedR
 				Cursor->SetPosition(Event.motion.x, Event.motion.y);
 				MessageHandler->OnRawMouseMove(Event.motion.xrel, Event.motion.yrel);
 				MessageHandler->OnMouseMove(); 
+				UE_LOG(LogHTML5Input, Verbose, TEXT("MouseMotion Pos(%d, %d) XRel:%d YRel:%d"), Event.motion.x, Event.motion.y, Event.motion.xrel, Event.motion.yrel);
 			} 
 			break;
 		case SDL_MOUSEWHEEL:
@@ -242,6 +247,7 @@ void FHTML5InputInterface::Tick(float DeltaTime, const SDL_Event& Event,TSharedR
 				SDL_MouseWheelEvent* w = (SDL_MouseWheelEvent*)&Event;
 				const float SpinFactor = 1 / 120.0f;
 				MessageHandler->OnMouseWheel(w->y * SpinFactor);
+				UE_LOG(LogHTML5Input, Verbose, TEXT("MouseWheel %f"), SpinFactor);
 			}
 			break;
 		default:
diff --git a/Engine/Source/Runtime/Core/Private/HTML5/HTML5PlatformMisc.cpp b/Engine/Source/Runtime/Core/Private/HTML5/HTML5PlatformMisc.cpp
--- a/Engine/Source/Runtime/Core/Private/HTML5/HTML5PlatformMisc.cpp
+++ b/Engine/Source/Runtime/Core/Private/HTML5/HTML5PlatformMisc.cpp
@@ -11,6 +11,8 @@
 #include "HTML5JavaScriptFx.h"
 #endif 
 
+#include "trace.h"
+
 #include "unicode/locid.h"
 #include "GenericPlatformContext.h"
 #include <SDL.h>
@@ -138,6 +140,10 @@ extern "C"
 	// callback from javascript. 
 	void on_fatal(const char* msg, const char* error)
 	{
+#ifdef __EMSCRIPTEN_TRACING__
+		emscripten_log(EM_LOG_CONSOLE, "Fatal Error: Closing trace!");
+		emscripten_trace_close();
+#endif
 		// !!JM todo: pass msg & error to a crash context? Must be copied?
 		if (GHTML5CrashHandler)
 		{
diff --git a/Engine/Source/Runtime/Core/Private/HTML5/HTML5PlatformOutputDevices.cpp b/Engine/Source/Runtime/Core/Private/HTML5/HTML5PlatformOutputDevices.cpp
new file mode 100644
--- /dev/null
+++ b/Engine/Source/Runtime/Core/Private/HTML5/HTML5PlatformOutputDevices.cpp
@@ -0,0 +1,18 @@
+// Copyright 1998-2015 Epic Games, Inc. All Rights Reserved.
+
+#include "CorePrivatePCH.h"
+#include "trace.h"
+
+class FTraceOutputDevice : public FOutputDevice
+{
+	virtual void Serialize( const TCHAR* V, ELogVerbosity::Type Verbosity, const class FName& Category )
+	{
+		emscripten_trace_log_message(Category.GetPlainANSIString(), TCHAR_TO_ANSI(V));
+	}
+};
+
+FOutputDevice* FHTML5PlatformOutputDevices::GetLog()
+{
+	static FTraceOutputDevice Singleton;
+	return &Singleton;
+}
\ No newline at end of file
diff --git a/Engine/Source/Runtime/Core/Private/HTML5/HTML5PlatformStackWalk.cpp b/Engine/Source/Runtime/Core/Private/HTML5/HTML5PlatformStackWalk.cpp
--- a/Engine/Source/Runtime/Core/Private/HTML5/HTML5PlatformStackWalk.cpp
+++ b/Engine/Source/Runtime/Core/Private/HTML5/HTML5PlatformStackWalk.cpp
@@ -133,4 +133,19 @@ void FHTML5PlatformStackWalk::CaptureStackBackTrace(uint64* BackTrace,uint32 Max
 #else
 	BackTrace[0] = 0;
 #endif
+}
+
+int32 FHTML5PlatformStackWalk::GetStackBackTraceString(char* OutputString, int32 MaxLen)
+{
+#if PLATFORM_HTML5_BROWSER
+	auto l = emscripten_get_callstack(EM_LOG_C_STACK|EM_LOG_DEMANGLE, OutputString, MaxLen);
+	char* p = OutputString;
+	while (*p && l) {
+		if (*p == '\n') *p = ' ';
+		p++;
+	}
+	return l;
+#else
+	return 0;
+#endif
 }
\ No newline at end of file
diff --git a/Engine/Source/Runtime/Core/Private/HTML5/HTML5Window.cpp b/Engine/Source/Runtime/Core/Private/HTML5/HTML5Window.cpp
--- a/Engine/Source/Runtime/Core/Private/HTML5/HTML5Window.cpp
+++ b/Engine/Source/Runtime/Core/Private/HTML5/HTML5Window.cpp
@@ -8,6 +8,8 @@
 #endif
 #include <SDL.h>
 
+DEFINE_LOG_CATEGORY_STATIC(LogHTML5Window, Log, All);
+
 FHTML5Window::~FHTML5Window()
 {
 	//    Use NativeWindow_Destroy() instead.
@@ -48,9 +50,11 @@ FPlatformRect FHTML5Window::GetScreenRect()
 #if !PLATFORM_HTML5_WIN32
 	int fs;
 	emscripten_get_canvas_size(&Width, &Height, &fs);
+	UE_LOG(LogHTML5Window, Verbose, TEXT("emscripten_get_canvas_size: Width:%d, Height:%d, Fullscreen:%d"), Width, Height, fs);
 #else
 	SDL_Window* WindowHandle= SDL_GL_GetCurrentWindow();
 	SDL_GetWindowSize(WindowHandle, &Width, &Height);
+	UE_LOG(LogHTML5Window, Verbose, TEXT("SDL_GetWindowSize: Width:%d, Height:%d"), Width, Height);
 #endif 
 	CalculateSurfaceSize(NULL,Width,Height);
 	ScreenRect.Right = Width;
diff --git a/Engine/Source/Runtime/Core/Public/HTML5/HTML5Application.h b/Engine/Source/Runtime/Core/Public/HTML5/HTML5Application.h
--- a/Engine/Source/Runtime/Core/Public/HTML5/HTML5Application.h
+++ b/Engine/Source/Runtime/Core/Public/HTML5/HTML5Application.h
@@ -37,4 +37,7 @@ private:
 	TSharedRef< class FGenericWindow > ApplicationWindow;
 
 	int32 WarmUpTicks; 
+
+	int32 WindowWidth;
+	int32 WindowHeight;
 };
\ No newline at end of file
diff --git a/Engine/Source/Runtime/Core/Public/HTML5/HTML5AssertionMacros.h b/Engine/Source/Runtime/Core/Public/HTML5/HTML5AssertionMacros.h
--- a/Engine/Source/Runtime/Core/Public/HTML5/HTML5AssertionMacros.h
+++ b/Engine/Source/Runtime/Core/Public/HTML5/HTML5AssertionMacros.h
@@ -28,10 +28,32 @@ void emscripten_log(int flags, ...);
 #define checkNoReentry(...)
 #define checkNoRecursion(...)
 
-#define check(expr)			{ if (!(expr)) { emscripten_log(255, "Expression '" #expr "' failed in " __FILE__ ":" PREPROCESSOR_TO_STRING(__LINE__) "!\n"); } }
-#define checkf(expr, ...)	{ if (!(expr)) { emscripten_log(255, "Expression '" #expr "' failed in " __FILE__ ":" PREPROCESSOR_TO_STRING(__LINE__) "!\n"); emscripten_log(255, ##__VA_ARGS__); FDebug::AssertFailed( #expr, __FILE__, __LINE__, ##__VA_ARGS__ ); } CA_ASSUME(expr); }
-#define verify(expr)		{ if(!(expr)) { emscripten_log(255, "Expression '" #expr "' failed in " __FILE__ ":" PREPROCESSOR_TO_STRING(__LINE__) "!\n"); } }
-#define verifyf(expr, ...)	{ if(!(expr)) { emscripten_log(255, "Expression '" #expr "' failed in " __FILE__ ":" PREPROCESSOR_TO_STRING(__LINE__) "!\n"); emscripten_log(255, ##__VA_ARGS__); } }
+#define html5_break_msg(msg, file, line) EM_ASM(alert('Expression ('+msg+') failed in '+file+':'+line+'!\nCheck console for details.\n'))
+
+#define check(expr)			{ if (!(expr)) { \
+		emscripten_log(255, "Expression '" #expr "' failed in " __FILE__ ":" PREPROCESSOR_TO_STRING(__LINE__) "!\n"); \
+		html5_break_msg(#expr, PREPROCESSOR_TO_STRING(__FILE__), PREPROCESSOR_TO_STRING(__LINE__)); \
+	} \
+}
+#define checkf(expr, ...)	{ if (!(expr)) { \
+		emscripten_log(255, "Expression '" #expr "' failed in " __FILE__ ":" PREPROCESSOR_TO_STRING(__LINE__) "!\n"); \
+		emscripten_log(255, ##__VA_ARGS__); \
+		FDebug::AssertFailed( #expr, __FILE__, __LINE__, ##__VA_ARGS__ ); \
+		html5_break_msg(#expr, PREPROCESSOR_TO_STRING(__FILE__), PREPROCESSOR_TO_STRING(__LINE__)); \
+	} \
+	CA_ASSUME(expr); \
+}
+#define verify(expr)		{ if(!(expr)) {\
+		emscripten_log(255, "Expression '" #expr "' failed in " __FILE__ ":" PREPROCESSOR_TO_STRING(__LINE__) "!\n"); \
+		html5_break_msg(#expr, PREPROCESSOR_TO_STRING(__FILE__), PREPROCESSOR_TO_STRING(__LINE__)); \
+	} \
+}
+#define verifyf(expr, ...)	{ if(!(expr)) { \
+		emscripten_log(255, "Expression '" #expr "' failed in " __FILE__ ":" PREPROCESSOR_TO_STRING(__LINE__) "!\n"); \
+		emscripten_log(255, ##__VA_ARGS__); \
+		html5_break_msg(#expr, PREPROCESSOR_TO_STRING(__FILE__), PREPROCESSOR_TO_STRING(__LINE__)); \
+	} \
+}
 
 #endif
 
@@ -41,8 +63,25 @@ void emscripten_log(int flags, ...);
 #undef checkfSlow
 #undef verifySlow
 
-#define checkSlow(expr, ...)   {if(!(expr)) { emscripten_log(255, "Expression '" #expr "' failed in " __FILE__ ":" PREPROCESSOR_TO_STRING(__LINE__) "!\n"); emscripten_log(255, ##__VA_ARGS__); FDebug::AssertFailed( #expr, __FILE__, __LINE__ ); } CA_ASSUME(expr); }
-#define checkfSlow(expr, ...)	{ if(!(expr)) { emscripten_log(255, "Expression '" #expr "' failed in " __FILE__ ":" PREPROCESSOR_TO_STRING(__LINE__) "!\n"); emscripten_log(255, ##__VA_ARGS__); FDebug::AssertFailed( #expr, __FILE__, __LINE__, __VA_ARGS__ ); } CA_ASSUME(expr); }
-#define verifySlow(expr)  {if(!(expr)) { emscripten_log(255, "Expression '" #expr "' failed in " __FILE__ ":" PREPROCESSOR_TO_STRING(__LINE__) "!\n"); FDebug::AssertFailed( #expr, __FILE__, __LINE__ ); } }
+#define checkSlow(expr, ...)   {if(!(expr)) { \
+		emscripten_log(255, "Expression '" #expr "' failed in " __FILE__ ":" PREPROCESSOR_TO_STRING(__LINE__) "!\n"); \
+		emscripten_log(255, ##__VA_ARGS__); FDebug::AssertFailed( #expr, __FILE__, __LINE__ ); \
+		html5_break_msg(#expr, PREPROCESSOR_TO_STRING(__FILE__), PREPROCESSOR_TO_STRING(__LINE__)); \
+	} \
+	CA_ASSUME(expr); \
+}
+#define checkfSlow(expr, ...)	{ if(!(expr)) { \
+		emscripten_log(255, "Expression '" #expr "' failed in " __FILE__ ":" PREPROCESSOR_TO_STRING(__LINE__) "!\n"); \
+		emscripten_log(255, ##__VA_ARGS__); FDebug::AssertFailed( #expr, __FILE__, __LINE__, __VA_ARGS__ ); \
+		html5_break_msg(#expr, PREPROCESSOR_TO_STRING(__FILE__), PREPROCESSOR_TO_STRING(__LINE__)); \
+	} \
+	CA_ASSUME(expr); \
+}
+#define verifySlow(expr)  {if(!(expr)) { \
+		emscripten_log(255, "Expression '" #expr "' failed in " __FILE__ ":" PREPROCESSOR_TO_STRING(__LINE__) "!\n"); \
+		FDebug::AssertFailed( #expr, __FILE__, __LINE__ ); \
+		html5_break_msg(#expr, PREPROCESSOR_TO_STRING(__FILE__), PREPROCESSOR_TO_STRING(__LINE__)); \
+	} \
+}
 
 #endif
diff --git a/Engine/Source/Runtime/Core/Public/HTML5/HTML5PlatformMisc.h b/Engine/Source/Runtime/Core/Public/HTML5/HTML5PlatformMisc.h
--- a/Engine/Source/Runtime/Core/Public/HTML5/HTML5PlatformMisc.h
+++ b/Engine/Source/Runtime/Core/Public/HTML5/HTML5PlatformMisc.h
@@ -10,6 +10,7 @@
 #include "HAL/Platform.h"
 #include "HTML5/HTML5DebugLogging.h"
 #include "HTML5/HTML5SystemIncludes.h"
+#include <emscripten.h>
 
 /**
  * HTML5 implementation of the misc OS functions
@@ -49,6 +50,9 @@ struct CORE_API FHTML5Misc : public FGenericPlatformMisc
 			__debugbreak();
 #else
 			emscripten_log(255, "DebugBreak() called!");
+			EM_ASM(
+				throw new Error('DebugBreak() called! Check Log for information');
+			);
 #endif
 		}
 	}
diff --git a/Engine/Source/Runtime/Core/Public/HTML5/HTML5PlatformOutputDevices.h b/Engine/Source/Runtime/Core/Public/HTML5/HTML5PlatformOutputDevices.h
--- a/Engine/Source/Runtime/Core/Public/HTML5/HTML5PlatformOutputDevices.h
+++ b/Engine/Source/Runtime/Core/Public/HTML5/HTML5PlatformOutputDevices.h
@@ -8,5 +8,10 @@
 #pragma once
 #include "GenericPlatform/GenericPlatformOutputDevices.h"
 
+struct CORE_API FHTML5PlatformOutputDevices : FGenericPlatformOutputDevices 
+{
+	static FOutputDevice*				GetLog();
+};
+
 // generic version
-typedef FGenericPlatformOutputDevices FPlatformOutputDevices;
+typedef FHTML5PlatformOutputDevices FPlatformOutputDevices;
diff --git a/Engine/Source/Runtime/Core/Public/HTML5/HTML5PlatformStackWalk.h b/Engine/Source/Runtime/Core/Public/HTML5/HTML5PlatformStackWalk.h
--- a/Engine/Source/Runtime/Core/Public/HTML5/HTML5PlatformStackWalk.h
+++ b/Engine/Source/Runtime/Core/Public/HTML5/HTML5PlatformStackWalk.h
@@ -17,6 +17,8 @@ struct CORE_API FHTML5PlatformStackWalk : public FGenericPlatformStackWalk
 
 	static void ProgramCounterToSymbolInfo(uint64 ProgramCounter,FProgramCounterSymbolInfo& out_SymbolInfo);
 	static void CaptureStackBackTrace(uint64* BackTrace,uint32 MaxDepth,void* Context = nullptr);
+	//Additional
+	static int32 GetStackBackTraceString(char* OutputString, int32 MaxLen);
 };
 
 typedef FHTML5PlatformStackWalk FPlatformStackWalk;
diff --git a/Engine/Source/Runtime/Launch/Private/HTML5/LaunchHTML5.cpp b/Engine/Source/Runtime/Launch/Private/HTML5/LaunchHTML5.cpp
--- a/Engine/Source/Runtime/Launch/Private/HTML5/LaunchHTML5.cpp
+++ b/Engine/Source/Runtime/Launch/Private/HTML5/LaunchHTML5.cpp
@@ -14,6 +14,7 @@ DEFINE_LOG_CATEGORY_STATIC(LogHTML5Launch, Log, All);
 
 #include "PlatformIncludes.h"
 #include <string.h>
+#include "trace.h"
 
 FEngineLoop	GEngineLoop;
 TCHAR GCmdLine[2048];
@@ -24,19 +25,43 @@ TCHAR GCmdLine[2048];
 
 void HTML5_Tick()
 {
+	static uint32 FrameCount = 1;
+	char Buf[128];
+	sprintf(Buf, "Frame %u", FrameCount);
+	emscripten_trace_record_frame_start();
+	emscripten_trace_enter_context(Buf);
 	GEngineLoop.Tick();
+	emscripten_trace_exit_context();
+	emscripten_trace_record_frame_end();
+
+	// Assuming 60fps, log out the memory report. This isn't important to be
+	// exactly every second, it just needs done periodically
+	if (FrameCount % 60 == 0)
+	{
+		emscripten_trace_report_memory_layout();
+	}
+
+	++FrameCount;
 }
 
 void HTML5_Init () 
 {	
+	emscripten_trace_record_frame_start();
+
 	// initialize the engine
 	UE_LOG(LogTemp,Display,TEXT("PreInit Start"));
+	emscripten_trace_enter_context("PreInit");
 	GEngineLoop.PreInit(0, NULL, GCmdLine);
+	emscripten_trace_exit_context();
 	UE_LOG(LogHTML5Launch,Display,TEXT("PreInit Complete"));
 	UE_LOG(LogHTML5Launch,Display,TEXT("Init Start"));
+	emscripten_trace_enter_context("Init");
 	GEngineLoop.Init();
+	emscripten_trace_exit_context();
 	UE_LOG(LogHTML5Launch,Display,TEXT("Init Complete"));
 
+	emscripten_trace_record_frame_end();
+
 
 #if PLATFORM_HTML5_WIN32
 	while (!GIsRequestingExit)
@@ -50,6 +75,26 @@ void HTML5_Init ()
 
 }
 
+class FHTML5Exec : private FSelfRegisteringExec
+{
+public:
+
+	FHTML5Exec()
+		{}
+
+	/** Console commands **/
+	virtual bool Exec( UWorld* InWorld, const TCHAR* Cmd, FOutputDevice& Ar ) override
+	{
+		if (FParse::Command(&Cmd, TEXT("em_trace_close")))
+		{
+			emscripten_trace_exit_context();//"main");
+			emscripten_trace_close();
+			return true;
+		}
+		return false;
+	}
+};
+static TAutoPtr<FHTML5Exec> GHTML5Exec;
 
 int main(int argc, char* argv[])
 {
@@ -63,6 +108,12 @@ int main(int argc, char* argv[])
 	
 #endif
 
+	// TODO: configure this via the command line?
+	emscripten_trace_configure("http://127.0.0.1:5000/", "UE4Game");
+	GHTML5Exec = new FHTML5Exec();
+
+	emscripten_trace_enter_context("main");
+
 	SDL_Init(SDL_INIT_VIDEO | SDL_INIT_AUDIO | SDL_INIT_NOPARACHUTE);
  	GCmdLine[0] = 0;
 	// to-do use Platform Str functions. 
@@ -81,4 +132,5 @@ int main(int argc, char* argv[])
 	return 0;
 }
 
+
 void EmptyLinkFunctionForStaticInitializationHTML5Win32(void) {} 
diff --git a/Engine/Source/Runtime/OpenGLDrv/Private/HTML5/HTML5OpenGL.cpp b/Engine/Source/Runtime/OpenGLDrv/Private/HTML5/HTML5OpenGL.cpp
--- a/Engine/Source/Runtime/OpenGLDrv/Private/HTML5/HTML5OpenGL.cpp
+++ b/Engine/Source/Runtime/OpenGLDrv/Private/HTML5/HTML5OpenGL.cpp
@@ -6,6 +6,8 @@
 #include <emscripten.h>
 #endif
 
+DEFINE_LOG_CATEGORY_STATIC(LogHTML5OpenGL, Log, All);
+
 bool FHTML5OpenGL::bCombinedDepthStencilAttachment = false;
 
 
@@ -255,6 +257,7 @@ void PlatformResizeGLContext( FPlatformOpenGLDevice* Device, FPlatformOpenGLCont
 	check(Context);
 	VERIFY_GL_SCOPE();
 
+	UE_LOG(LogHTML5OpenGL, Verbose, TEXT("SDL_SetWindowSize(%d,%d)"), SizeX, SizeY);
 	SDL_SetWindowSize(Device->WindowHandle,SizeX,SizeY); 
 
 	glViewport(0, 0, SizeX, SizeY);
@@ -301,6 +304,7 @@ void PlatformGetBackbufferDimensions( uint32& OutWidth, uint32& OutHeight )
 	check(Surface);
 	OutWidth  = Surface->w;
 	OutHeight = Surface->h;
+	UE_LOG(LogHTML5OpenGL, Verbose, TEXT("PlatformGetBackbufferDimensions(%d, %d)"), OutWidth, OutHeight);
 }
 
 // =============================================================
@@ -312,6 +316,7 @@ bool PlatformContextIsCurrent( uint64 QueryContext )
 
 FRHITexture* PlatformCreateBuiltinBackBuffer(FOpenGLDynamicRHI* OpenGLRHI, uint32 SizeX, uint32 SizeY)
 {
+	UE_LOG(LogHTML5OpenGL, Verbose, TEXT("PlatformCreateBuiltinBackBuffer(%d, %d)"), SizeX, SizeY);
 	uint32 Flags = TexCreate_RenderTargetable;
 	FOpenGLTexture2D* Texture2D = new FOpenGLTexture2D(OpenGLRHI, 0, GL_RENDERBUFFER, GL_COLOR_ATTACHMENT0, SizeX, SizeY, 0, 1, 1, 1, PF_B8G8R8A8, false, false, Flags, nullptr);
 	OpenGLTextureAllocated(Texture2D, Flags);
@@ -337,7 +342,8 @@ extern "C"
 	// callback from javascript. 
 	void resize_game(int w, int h)
 	{
-      // to-do: remove this separate code path. 
+      	// to-do: remove this separate code path. 
+		UE_LOG(LogHTML5OpenGL, Verbose, TEXT("resize_game(%d, %d) callback"), w, h);
 	}
 #endif 
 }
diff --git a/Engine/Source/Runtime/OpenGLDrv/Private/OpenGLShaders.cpp b/Engine/Source/Runtime/OpenGLDrv/Private/OpenGLShaders.cpp
--- a/Engine/Source/Runtime/OpenGLDrv/Private/OpenGLShaders.cpp
+++ b/Engine/Source/Runtime/OpenGLDrv/Private/OpenGLShaders.cpp
@@ -696,6 +696,16 @@ ShaderType* CompileOpenGLShader(const TArray<uint8>& Code)
 		// On Android the same shader is compiled with different hacks to find the right one(s) to apply so don't cache unless successful
 		glGetShaderiv(Resource, GL_COMPILE_STATUS, &CompileStatus);
 #endif
+#if PLATFORM_HTML5
+		glGetShaderiv(Resource, GL_COMPILE_STATUS, &CompileStatus);
+		if (CompileStatus == GL_FALSE)
+		{
+			char Msg[2048];
+			glGetShaderInfoLog(Resource, 2048, nullptr, Msg);
+			UE_LOG(LogRHI, Error, TEXT("Shader compile failed: %s\n Original Source is (len %d) %s"), ANSI_TO_TCHAR(Msg), GlslCodeLength, ANSI_TO_TCHAR(GlslCodeString));
+		}
+#endif
+
 		if ( CompileStatus == GL_TRUE )
 		{
 			if ( FOpenGL::SupportsSeparateShaderObjects() )
