Remove dependency from HTML5LaunchHelper on UnrealBuildTool. It prevents debugging the build process, because cleaning HTML5LaunchHelper also cleans UBT (which is already open in the debugger).

diff --git a/Engine/Source/Programs/HTML5/HTML5LaunchHelper/HTML5LaunchHelper.csproj b/Engine/Source/Programs/HTML5/HTML5LaunchHelper/HTML5LaunchHelper.csproj
--- a/Engine/Source/Programs/HTML5/HTML5LaunchHelper/HTML5LaunchHelper.csproj
+++ b/Engine/Source/Programs/HTML5/HTML5LaunchHelper/HTML5LaunchHelper.csproj
@@ -55,11 +55,4 @@
   <Target Name="AfterBuild">
   </Target>
   -->
-  <ItemGroup>
-    <ProjectReference Include="..\..\UnrealBuildTool\UnrealBuildTool.csproj">
-      <Project>{fd7c5e1a-cfe4-4fd5-a525-1eb1599a39ac}</Project>
-      <Name>UnrealBuildTool</Name>
-      <Private>False</Private>
-    </ProjectReference>
-  </ItemGroup>
 </Project>
\ No newline at end of file
diff --git a/Engine/Source/Programs/HTML5/HTML5LaunchHelper/Program.cs b/Engine/Source/Programs/HTML5/HTML5LaunchHelper/Program.cs
--- a/Engine/Source/Programs/HTML5/HTML5LaunchHelper/Program.cs
+++ b/Engine/Source/Programs/HTML5/HTML5LaunchHelper/Program.cs
@@ -4,7 +4,7 @@ using System.Linq;
 using System.Text;
 using System.Threading.Tasks;
 using System.Diagnostics;
-using UnrealBuildTool;
+using System.IO;
 
 // Example Command Lines
 // HTML5LaunchHelper -m -h -b /Applications/Safari.app localhost:8000
@@ -25,10 +25,23 @@ namespace HTML5LaunchHelper
 		static private List<Process> ProcessesToKill = new List<Process>();
 		static private List<Process> ProcessesToWatch = new List<Process>();
 
+		static bool IsRunningOnMac()
+		{
+			PlatformID Platform = Environment.OSVersion.Platform;
+			switch (Platform)
+			{
+				case PlatformID.Unix:
+					return System.IO.File.Exists("/System/Library/CoreServices/SystemVersion.plist");
+				case PlatformID.MacOSX:
+					return true;
+			}
+			return false;
+		}
+
 		static Process SpawnBrowserProcess(string bpath, string args)
 		{
 			var Result = new Process();
-			if (BuildHostPlatform.Current.Platform == UnrealTargetPlatform.Mac)
+			if (IsRunningOnMac())
 			{
 	            Result.StartInfo.FileName = "/usr/bin/open";
 				Result.StartInfo.UseShellExecute = false;
