UEPLAT-794 Emit compressed javascript from HTML5 packaging .

  - big build products which are sent over the wire are now compressed.
  - modfied HTML5Launcher to serve compressed files. Added this to the final packaged game too because python simple http server does not handle compressed files.
  - combined various javascript utlity files into one.

diff --git a/Engine/Build/HTML5/Game.html.template b/Engine/Build/HTML5/Game.html.template
--- a/Engine/Build/HTML5/Game.html.template
+++ b/Engine/Build/HTML5/Game.html.template
@@ -276,11 +276,9 @@ if (webgl_detect())
 Module.setStatus('Downloading...');
 
   // add these scripts to the dom. 
-  loadJavaScriptFile("json2.js",false);
-  loadJavaScriptFile("moz_binarystring.js",false);
-  loadJavaScriptFile("jStorage.js",false);
-  loadJavaScriptFile("%GAME%.data.js",false);
-  loadJavaScriptFile("%CONFIG%.js",true);
+  loadJavaScriptFile("Utility.js.gz",false);
+  loadJavaScriptFile("%GAME%.data.js.gz",false);
+  loadJavaScriptFile("%CONFIG%.js.gz",true);
 
 }
 else
diff --git a/Engine/Build/HTML5/GameX.html.template b/Engine/Build/HTML5/GameX.html.template
--- a/Engine/Build/HTML5/GameX.html.template
+++ b/Engine/Build/HTML5/GameX.html.template
@@ -429,11 +429,9 @@ $(document).ready(function() {
     addLog( "Starting downloads" );
     // start downloads in parallel, print when done.  
     var promises = [
-			'json2.js',
-			'jStorage.js',
-			'moz_binarystring.js',
-			'%GAME%.data.js',
-			'%CONFIG%.js'   
+			'Utility.js.gz',
+			'%GAME%.data.js.gz',
+			'%CONFIG%.js.gz'   
 		].map(download); 
 
 	// add the data file promise. 
diff --git a/Engine/Build/HTML5/Readme.txt b/Engine/Build/HTML5/Readme.txt
new file mode 100644
--- /dev/null
+++ b/Engine/Build/HTML5/Readme.txt
@@ -0,0 +1,27 @@
+This folder contains pre compressed files ready for deployment, please make sure to use the shipping build for final deployment. 
+
+Files Required for final deployment. 
+---------------------------------
+
+*.js.gz -compressed javascript files.
+*.data  -compressed game content
+*.mem	-compressed memory iniitialization file
+*.html 	-uncompressed landing page. 
+
+*.symbols - uncompressed symbols, if necessary. 
+
+
+Local Testing
+------------- 
+
+run HTML5Launcher.exe (via mono on Mac) to start a web sever which is conigured to serve compressed files on localhost. This is not a production quality server.
+add -ServerPort=XXXX to the command line if necessary to change the serving port. Default port is 8000. 
+
+
+How to deploy on Apache 
+------------------------ 
+
+Todo. 
+
+
+
diff --git a/Engine/Source/Programs/AutomationTool/HTML5/HTML5.Automation.csproj b/Engine/Source/Programs/AutomationTool/HTML5/HTML5.Automation.csproj
--- a/Engine/Source/Programs/AutomationTool/HTML5/HTML5.Automation.csproj
+++ b/Engine/Source/Programs/AutomationTool/HTML5/HTML5.Automation.csproj
@@ -34,6 +34,9 @@
     <TreatWarningsAsErrors>true</TreatWarningsAsErrors>
   </PropertyGroup>
   <ItemGroup>
+    <Reference Include="Ionic.Zip.Reduced">
+      <HintPath>..\..\..\..\Binaries\DotNET\Ionic.Zip.Reduced.dll</HintPath>
+    </Reference>
     <Reference Include="OneSky, Version=1.0.0.0, Culture=neutral, processorArchitecture=MSIL">
       <HintPath>..\..\..\..\Binaries\DotNET\OneSky.dll</HintPath>
       <SpecificVersion>False</SpecificVersion>
diff --git a/Engine/Source/Programs/AutomationTool/HTML5/HTML5Platform.Automation.cs b/Engine/Source/Programs/AutomationTool/HTML5/HTML5Platform.Automation.cs
--- a/Engine/Source/Programs/AutomationTool/HTML5/HTML5Platform.Automation.cs
+++ b/Engine/Source/Programs/AutomationTool/HTML5/HTML5Platform.Automation.cs
@@ -9,6 +9,7 @@ using AutomationTool;
 using UnrealBuildTool;
 using System.Diagnostics;
 using System.Threading;
+using Ionic.Zip;
 
 public class HTML5Platform : Platform
 {
@@ -73,6 +74,12 @@ public class HTML5Platform : Platform
             }
         }
 
+		// packaging created 2 files - .data.js and .data file. lets compress them.
+
+		CompressFile(FinalDataLocation, FinalDataLocation);
+		CompressFile(FinalDataLocation + ".js", FinalDataLocation + ".js.gz");
+		File.Delete(FinalDataLocation + ".js");
+
         // copy the "Executable" to the package directory
         string GameExe = Path.GetFileNameWithoutExtension(Params.ProjectGameExeFilename);
         if (Params.ClientConfigsToBuild[0].ToString() != "Development")
@@ -92,14 +99,15 @@ public class HTML5Platform : Platform
 
         if (Path.Combine(Path.GetDirectoryName(Params.ProjectGameExeFilename), GameExe) != Path.Combine(PackagePath, GameExe))
         {
-            File.Copy(Path.Combine(Path.GetDirectoryName(Params.ProjectGameExeFilename), GameExe), Path.Combine(PackagePath, GameExe), true);
-            File.Copy(Path.Combine(Path.GetDirectoryName(Params.ProjectGameExeFilename), GameExe) + ".mem", Path.Combine(PackagePath, GameExe) + ".mem", true);
+			// compress all javascript files in place. 
+			Log("Compressing and Copying main js executable");
+			CompressFile(FullGameExePath, Path.Combine(PackagePath, GameExe) + ".gz");
+			Log("Compressing and copying memory Initialization file");
+			CompressFile(Path.Combine(Path.GetDirectoryName(Params.ProjectGameExeFilename), GameExe) + ".mem", (Path.Combine(PackagePath, GameExe) + ".mem"));
 			File.Copy(Path.Combine(Path.GetDirectoryName(Params.ProjectGameExeFilename), GameExe) + ".symbols", Path.Combine(PackagePath, GameExe) + ".symbols", true);
         }
-        File.SetAttributes(Path.Combine(PackagePath, GameExe), FileAttributes.Normal);
-        File.SetAttributes(Path.Combine(PackagePath, GameExe) + ".mem", FileAttributes.Normal);
-		File.SetAttributes(Path.Combine(PackagePath, GameExe) + ".symbols", FileAttributes.Normal);
 
+		File.SetAttributes(Path.Combine(PackagePath, GameExe) + ".symbols", FileAttributes.Normal);
 
         // put the HTML file to the package directory
 
@@ -134,18 +142,63 @@ public class HTML5Platform : Platform
 
         GenerateFileFromTemplate(TemplateFile, OutputFile, Params.ShortProjectName, Params.ClientConfigsToBuild[0].ToString(), Params.StageCommandline, !Params.IsCodeBasedProject, HeapSize);
 
-        // copy the jstorage files to the binaries directory
         string JSDir = Path.Combine(CombinePaths(CmdEnv.LocalRoot, "Engine"), "Build", "HTML5");
         string OutDir = PackagePath;
-        File.Copy(JSDir + "/json2.js", OutDir + "/json2.js", true);
-        File.SetAttributes(OutDir + "/json2.js", FileAttributes.Normal);
-        File.Copy(JSDir + "/jStorage.js", OutDir + "/jStorage.js", true);
-        File.SetAttributes(OutDir + "/jStorage.js", FileAttributes.Normal);
-        File.Copy(JSDir + "/moz_binarystring.js", OutDir + "/moz_binarystring.js", true);
-        File.SetAttributes(OutDir + "/moz_binarystring.js", FileAttributes.Normal);
+
+		// Gather utlity .js files and combine into one file
+		string[] UtilityJavaScriptFiles = Directory.GetFiles(JSDir, "*.js");
+
+		string DestinationFile = OutDir + "/Utility.js";
+		foreach( var UtilityFile in UtilityJavaScriptFiles)
+		{
+			string Data = File.ReadAllText(UtilityFile);
+			File.AppendAllText(DestinationFile, Data);
+		}
+
+		CompressFile(DestinationFile, DestinationFile + ".gz");
+
+		// delete uncompressed file.
+		if (File.Exists(DestinationFile))
+		{
+			File.Delete(DestinationFile);
+		}
+
         PrintRunTime();
 	}
 
+	void CompressFile(string Source, string Destination)
+	{
+		bool DeleteSource = false; 
+
+		if(  Source == Destination )
+		{
+			string CopyOrig = Source + ".Copy";
+			File.Copy(Source, CopyOrig);
+			Source = CopyOrig;
+			DeleteSource = true; 
+		}
+
+		using (System.IO.Stream input = System.IO.File.OpenRead(Source))
+		{
+			using (var raw = System.IO.File.Create(Destination))
+				{
+					using (Stream compressor = new Ionic.Zlib.GZipStream(raw, Ionic.Zlib.CompressionMode.Compress,Ionic.Zlib.CompressionLevel.BestCompression))
+					{
+						byte[] buffer = new byte[2048];
+						int SizeRead = 0;
+						while ((SizeRead = input.Read(buffer, 0, buffer.Length)) != 0)
+						{
+							compressor.Write(buffer, 0, SizeRead);
+						}
+					}
+				}
+		}
+
+		if (DeleteSource && File.Exists(Source))
+		{
+			File.Delete(Source);
+		}
+	}
 
     public override bool RequiresPackageToDeploy
     {
@@ -247,26 +300,28 @@ public class HTML5Platform : Platform
 		string OutputFile = Path.Combine(PackagePath, (Params.ClientConfigsToBuild[0].ToString() != "Development" ? (Params.ShortProjectName + "-HTML5-" + Params.ClientConfigsToBuild[0].ToString()) : Params.ShortProjectName)) + ".html";
 
 		SC.ArchiveFiles(PackagePath, Path.GetFileName(FinalDataLocation));
-		SC.ArchiveFiles(PackagePath, Path.GetFileName(FinalDataLocation + ".js"));
-		SC.ArchiveFiles(PackagePath, Path.GetFileName(GameExe));
+		SC.ArchiveFiles(PackagePath, Path.GetFileName(FinalDataLocation + ".js.gz"));
+		SC.ArchiveFiles(PackagePath, Path.GetFileName(GameExe +".gz"));
 		SC.ArchiveFiles(PackagePath, Path.GetFileName(GameExe + ".mem"));
 		SC.ArchiveFiles(PackagePath, Path.GetFileName(GameExe + ".symbols"));
-		SC.ArchiveFiles(PackagePath, Path.GetFileName("json2.js"));
-		SC.ArchiveFiles(PackagePath, Path.GetFileName("jStorage.js"));
-		SC.ArchiveFiles(PackagePath, Path.GetFileName("moz_binarystring.js"));
+		SC.ArchiveFiles(PackagePath, Path.GetFileName("Utility.js.gz"));
 		SC.ArchiveFiles(PackagePath, Path.GetFileName(OutputFile));
-        
 
-        if (HTMLPakAutomation.CanCreateMapPaks(Params))
-        {
-            // find all paks.
-            string[] Files = Directory.GetFiles(Path.Combine(PackagePath, Params.ShortProjectName), "*",SearchOption.AllDirectories);
-            foreach(string PakFile in Files)
-            {
-                var DestPak = PakFile.Replace(PackagePath,"");
-                SC.ArchivedFiles.Add(PakFile, DestPak);
-            }
-        }
+		// Archive HTML5 Server and a Readme. 
+		var LaunchHelperPath = CombinePaths(CmdEnv.LocalRoot, "Engine/Binaries/DotNET/");
+		SC.ArchiveFiles(LaunchHelperPath, "HTML5LaunchHelper.exe");
+		SC.ArchiveFiles(Path.Combine(CombinePaths(CmdEnv.LocalRoot, "Engine"), "Build", "HTML5"), "Readme.txt");
+
+		if (HTMLPakAutomation.CanCreateMapPaks(Params))
+		{
+			// find all paks.
+			string[] Files = Directory.GetFiles(Path.Combine(PackagePath, Params.ShortProjectName), "*",SearchOption.AllDirectories);
+			foreach(string PakFile in Files)
+			{
+				var DestPak = PakFile.Replace(PackagePath,"");
+				SC.ArchivedFiles.Add(PakFile, DestPak);
+			}
+		}
 
 	}
 
@@ -320,7 +375,6 @@ public class HTML5Platform : Platform
 		ConfigCache.GetInt32("/Script/HTML5PlatformEditor.HTML5TargetSettings", "DeployServerPort", out ServerPort);
 		string WorkingDirectory = Path.GetDirectoryName(ClientApp);
 		string url = Path.GetFileName(ClientApp) +".html";
-		string args = "-m ";
 		// Are we running via cook on the fly server?
 		// find our http url - This is awkward because RunClient doesn't have real information that NFS is running or not.
 		bool IsCookOnTheFly = false;
@@ -341,27 +395,29 @@ public class HTML5Platform : Platform
 		}
 		else
 		{
-			url = String.Format("http://127.0.0.1:{0}/{1}", ServerPort, url);
-			args += String.Format("-h -s {0} ", ServerPort);
+			url = String.Format("http://localhost:{0}/{1}", ServerPort, url);
 		}
 
 		// Check HTML5LaunchHelper source for command line args
+
 		var LowerBrowserPath = browserPath.ToLower();
-		args += String.Format("-b \"{0}\" -p \"{1}\" -w \"{2}\" ", browserPath, HTML5SDKInfo.PythonPath(), WorkingDirectory);
-		args += url + " ";
 		var ProfileDirectory = Path.Combine(Utils.GetUserSettingDirectory(), "UE4_HTML5", "user");
+
+		string BrowserCommandline = url; 
+
 		if (LowerBrowserPath.Contains("chrome"))
 		{
-			args += String.Format("--user-data-dir=\"{0}\" --enable-logging --no-first-run", Path.Combine(ProfileDirectory, "chrome"));
+			BrowserCommandline  += "  " + String.Format("--user-data-dir=\"{0}\" --enable-logging --no-first-run", Path.Combine(ProfileDirectory, "chrome"));
 		}
 		else if (LowerBrowserPath.Contains("firefox"))
 		{
-			args += String.Format("-no-remote -profile \"{0}\"", Path.Combine(ProfileDirectory, "firefox"));
+			BrowserCommandline += "  " +  String.Format("-no-remote -profile \"{0}\"", Path.Combine(ProfileDirectory, "firefox"));
 		}
-		//else if (browserPath.Contains("Safari")) {}
+
+		string LauncherArguments = string.Format( " -Browser=\"{0}\" + -BrowserCommandLine=\"{1}\" -ServerPort=\"{2}\" -ServerRoot=\"{3}\" ", new object[] { browserPath,BrowserCommandline,ServerPort,WorkingDirectory });
 
 		var LaunchHelperPath = CombinePaths(CmdEnv.LocalRoot, "Engine/Binaries/DotNET/HTML5LaunchHelper.exe");
-		ProcessResult BrowserProcess = Run(LaunchHelperPath, args, null, ClientRunFlags | ERunOptions.NoWaitForExit);
+		ProcessResult BrowserProcess = Run(LaunchHelperPath, LauncherArguments, null, ClientRunFlags | ERunOptions.NoWaitForExit);
 	
 		return BrowserProcess;
 
diff --git a/Engine/Source/Programs/HTML5/HTML5LaunchHelper/Program.cs b/Engine/Source/Programs/HTML5/HTML5LaunchHelper/Program.cs
--- a/Engine/Source/Programs/HTML5/HTML5LaunchHelper/Program.cs
+++ b/Engine/Source/Programs/HTML5/HTML5LaunchHelper/Program.cs
@@ -1,27 +1,338 @@
 ﻿using System;
+using System.Net; 
 using System.Collections.Generic;
 using System.Linq;
 using System.Text;
 using System.Threading.Tasks;
 using System.Diagnostics;
 using System.IO;
+using System.Reflection;
+using System.Web;
 
-// Example Command Lines
-// HTML5LaunchHelper -m -h -b /Applications/Safari.app localhost:8000
-// HTML5LaunchHelper -m -h -b /Applications/Firefox.app http://localhost:8000 -silent -profile [REDACTED_USER_HOME]/webporfiles/nightly/UE4HTML5 --jsconsole
-// HTML5LaunchHelper -m -h -b /Applications/Google\ Chrome.app http://localhost:8000 --user-data-dir="[REDACTED_USER_HOME]/webporfiles/chrome/UE4HTML5" --enable-logging --no-first-run
 namespace HTML5LaunchHelper
 {
+	class HttpServer
+	{
+#region extension to MIME type list
+		// some basic mime types, not really important but for completeness sake. 
+		private static IDictionary<string, string> MimeTypeMapping = new Dictionary<string, string>(StringComparer.InvariantCultureIgnoreCase) 
+		{
+			{".bin", "application/octet-stream"},
+			{".css", "text/css"},
+			{".dll", "application/octet-stream"},
+			{".dmg", "application/octet-stream"},
+			{".ear", "application/java-archive"},
+			{".eot", "application/octet-stream"},
+			{".exe", "application/octet-stream"},
+			{".flv", "video/x-flv"},
+			{".gif", "image/gif"},
+			{".hqx", "application/mac-binhex40"},
+			{".htc", "text/x-component"},
+			{".htm", "text/html"},
+			{".html", "text/html"},
+			{".ico", "image/x-icon"},
+			{".img", "application/octet-stream"},
+			{".iso", "application/octet-stream"},
+			{".jar", "application/java-archive"},
+			{".jardiff", "application/x-java-archive-diff"},
+			{".jng", "image/x-jng"},
+			{".jnlp", "application/x-java-jnlp-file"},
+			{".jpeg", "image/jpeg"},
+			{".jpg", "image/jpeg"},
+			{".js", "application/x-javascript"},
+			{".mml", "text/mathml"},
+			{".mng", "video/x-mng"},
+			{".mov", "video/quicktime"},
+			{".mp3", "audio/mpeg"},
+			{".mpeg", "video/mpeg"},
+			{".mpg", "video/mpeg"},
+			{".msi", "application/octet-stream"},
+			{".msm", "application/octet-stream"},
+			{".msp", "application/octet-stream"},
+			{".pdb", "application/x-pilot"},
+			{".pdf", "application/pdf"},
+			{".pem", "application/x-x509-ca-cert"},
+			{".pl", "application/x-perl"},
+			{".pm", "application/x-perl"},
+			{".png", "image/png"},
+			{".prc", "application/x-pilot"},
+			{".ra", "audio/x-realaudio"},
+			{".rar", "application/x-rar-compressed"},
+			{".rpm", "application/x-redhat-package-manager"},
+			{".rss", "text/xml"},
+			{".run", "application/x-makeself"},
+			{".sea", "application/x-sea"},
+			{".shtml", "text/html"},
+			{".sit", "application/x-stuffit"},
+			{".swf", "application/x-shockwave-flash"},
+			{".tcl", "application/x-tcl"},
+			{".tk", "application/x-tcl"},
+			{".txt", "text/plain"},
+			{".war", "application/java-archive"},
+			{".wbmp", "image/vnd.wap.wbmp"},
+			{".wmv", "video/x-ms-wmv"},
+			{".xml", "text/xml"},
+			{".xpi", "application/x-xpinstall"},
+			{".zip", "application/zip"},
+    };
+#endregion
+
+		private HttpListener WebServer = new HttpListener();
+		private string Root; 
+
+		public HttpServer(int Port, string ServerRoot)
+		{
+			Root = ServerRoot;
+			WebServer.Prefixes.Add("http://localhost:" + Port.ToString() + "/");
+		}
+
+		public void Run() 
+		{
+			System.Console.WriteLine("Starting Server at " + WebServer.Prefixes.First().ToString());
+			WebServer.Start();
+			Task.Factory.StartNew(()
+					=>
+						{
+							while( WebServer.IsListening)
+							{
+								// Handle requests in threaded mode.
+								Task.Factory.StartNew((Ctx)
+									 =>
+									 {
+										 var Context = Ctx as HttpListenerContext;
+										 try
+										 {
+											 RequestHandler(Context);
+										 }
+										 catch { }
+										 finally
+										 {
+
+											 Context.Response.Close();
+										 }
+
+									 }, WebServer.GetContext());
+							}
+						}
+				);
+		}
+
+		private void RequestHandler(HttpListenerContext Context)
+		{
+			if (Directory.Exists(Root + Context.Request.Url.LocalPath))
+			{
+				// Process the list of files found in the directory. 
+				string[] fileEntries = Directory.GetFileSystemEntries(Root + Context.Request.Url.LocalPath);
+				string Response =   "<html>\n" +
+									"<body>\n" +
+									"<h2>Unreal WebServer</h2>\n" +
+									"<h3>Directory listing for " + Context.Request.Url.LocalPath + "</h3>\n" +
+									"<hr>\n"; 
+
+				Response += "<table>\n"; 
+
+				Response +=  "<tr>\n" + 
+								"\t<th>Filename</th>\n"+
+								"\t<th>TimeStamp</th>\n"+
+							 "</tr>\n";
+
+				foreach (string fileName in fileEntries)
+				{
+					string Slash = Directory.Exists(fileName) ? "/" : "" ; 
+					Response += "<tr>\n\t<td><a href=\"" + fileName + "\">" + Path.GetFileName(fileName) + Slash + "</a></td>\n\t<td> " + File.GetLastAccessTime(fileName).ToString() + "</td>\n</tr>\n";
+				}
+				Response  += "</table></html>";
+				byte[] buf = Encoding.UTF8.GetBytes(Response);
+				Context.Response.ContentLength64 = buf.Length;
+				Context.Response.OutputStream.Write(buf, 0, buf.Length);
+				Context.Response.Close();
+			}
+			if (File.Exists(Root + Context.Request.Url.LocalPath))
+			{
+				string RequestedFile = Root + Context.Request.Url.LocalPath;
+				System.Console.WriteLine("Serving " + RequestedFile);
+				using (Stream source = File.OpenRead(RequestedFile))
+				{
+					string Extention = Path.GetExtension(RequestedFile);
+					string MimeType = "text/html";
+
+					if (MimeTypeMapping.ContainsKey(Extention))
+					{
+						MimeType = MimeTypeMapping[Extention];
+					}
+					// This is the crux of serving pre-compressed files. 
+					else if (Extention == ".gz" || Extention == ".data" || Extention == ".mem")
+					{
+						Context.Response.AddHeader("Content-Encoding", "gzip");
+					}
+
+					byte[] buffer = new byte[source.Length];
+					source.Read(buffer, 0, buffer.Length);
+					Context.Response.ContentType = MimeType;
+					Context.Response.ContentLength64 = buffer.Length;
+					Context.Response.OutputStream.Write(buffer, 0, buffer.Length);
+					Context.Response.Close();
+				}
+			}
+			Context.Response.Close();
+		}
+
+		public void Stop()
+		{
+			WebServer.Stop();
+			WebServer.Close();
+		}
+	}
+
+	class ArgumentName : Attribute
+	{
+		public string Name;
+		public ArgumentName(string _Name) 
+		{
+			Name = _Name;
+		}
+	}
+
+	class DefaultArgument: Attribute
+	{
+		public string Value;
+		public DefaultArgument(string _Value)
+		{
+			Value = _Value; 
+		}
+	}
+
+	// Various command line options supported. 
+	class Arguments
+	{
+		// if this is set -  This browser is spawned and the web server blocks till the browser quits. 
+		// if this not set - Just the server starts up and waits for key to quit. 
+		[ArgumentName("-Browser="), DefaultArgument("")]
+		public string Browser
+		{
+			get;
+			set; 
+		}
+
+		[ArgumentName("-ServerRoot="), DefaultArgument("./")]
+		public string ServerRoot
+		{
+			get;
+			set;
+		}
+
+		[ArgumentName("-ServerPort="), DefaultArgument("8000")]
+		public string ServerPort
+		{
+			get;
+			set;
+		}
+
+		[ArgumentName("-BrowserCommandLine="), DefaultArgument("")]
+		public string BrowserCommandLine
+		{
+			get;
+			set;
+		}
+
+		public Arguments()
+		{}
+
+		public bool Parse(string[] args)
+		{
+			PropertyInfo[] Infos = typeof(Arguments).GetProperties();
+			foreach( var Info in Infos )
+			{
+				object[] Attributes = Info.GetCustomAttributes(false);
+
+				string Name = null;
+				string DefaultValue = null; 
+
+				foreach( var Att in Attributes)
+				{
+					if (Att.GetType() == typeof(ArgumentName))
+					{
+						Name = (Att as ArgumentName).Name;
+					}
+					if (Att.GetType() == typeof(DefaultArgument))
+					{
+						DefaultValue = (Att as DefaultArgument).Value;
+					}
+				}
+
+				bool found = false; 
+				foreach(var arg in args)
+				{
+					if (arg.StartsWith(Name))
+					{
+						string val = arg.Replace(Name, "");
+						Info.SetValue(this, val);
+						found = true;
+						break;
+					}
+				}
+
+				if (!found && DefaultValue != null)
+				{
+					Info.SetValue(this, DefaultValue);
+				}
+				else if ( !found && DefaultValue == null)
+				{
+					return false; 
+				}
+			}
+			return true; 
+		}
+
+		public void ShowParsedValues()
+		{
+			PropertyInfo[] Infos = typeof(Arguments).GetProperties();
+			foreach (var Info in Infos)
+			{
+				object[] Attributes = Info.GetCustomAttributes(false);
+				string Name = null;
+				foreach (var Att in Attributes)
+				{
+					if (Att.GetType() == typeof(ArgumentName))
+					{
+						Name = (Att as ArgumentName).Name;
+					}
+				}
+				if (Info.GetValue(this) != null)
+				{
+					System.Console.WriteLine("Name: " + Name + " " + Info.GetValue(this).ToString());
+				}
+			}
+		}
+
+		public void ShowAllOptions()
+		{
+			PropertyInfo[] Infos = typeof(Arguments).GetProperties();
+			foreach (var Info in Infos)
+			{
+				object[] Attributes = Info.GetCustomAttributes(false);
+
+				string Name = null;
+				string DefaultValue = null;
+
+				foreach (var Att in Attributes)
+				{
+					if (Att.GetType() == typeof(ArgumentName))
+					{
+						Name = (Att as ArgumentName).Name;
+					}
+					if (Att.GetType() == typeof(DefaultArgument))
+					{
+						DefaultValue = (Att as DefaultArgument).Value;
+					}
+				}
+				System.Console.WriteLine("Option: {0}, Default Value {1}", Name, DefaultValue == null ? " None, this option is required " : DefaultValue);
+			}
+		}
+	}
+
 	class Program
 	{
-		static private bool 	bMultiprocessBrowser = false;
-		static private bool 	bStartPythonServer = false;
-		static private int  	PythonServerPort = 8000;
-		static private string 	PythonPath = "python";
-		static private string 	BrowserPath = null;
-		static private string	ServerDirectory = null;
-		static private string 	Url = null;
-		static private string   AdditionalArgs = "";
 		static private List<Process> ProcessesToKill = new List<Process>();
 		static private List<Process> ProcessesToWatch = new List<Process>();
 
@@ -59,37 +370,8 @@ namespace HTML5LaunchHelper
 				Result.StartInfo.Arguments = args;
 				Result.EnableRaisingEvents = true;
 			}
-			System.Console.WriteLine("Spawning Browser Process {0} with args {1}\n", bpath, args);
-			return Result;
-		}
-
-		static int Main(string[] args)
-		{
-			if (!ParseArgs(args))
-			{
-				// TODO: print error & options.
-				System.Console.Write("Invalid args\n");
-				return -1;
-			}
 
-			if (bStartPythonServer)
-			{
-				BeginPythonServer();
-				// Give the server time to start, 1 Second should be enough.
-				System.Threading.Thread.Sleep(1000);
-			}
-
-			var bIsSafari = BrowserPath.Contains("Safari");
-
-			// Browsers can be multiprocess programs (Chrome, basically)
-			// so we need to catch spawning of other child processes. The trick is
-			// they aren't really child-processes at all. There appears to be no real binding between the two,
-			// So we kind of fudged it a bit here.
-			var PrevProcesses = Process.GetProcesses();
-			var FirstProcess = SpawnBrowserProcess(BrowserPath, !bIsSafari ? Url + " " + AdditionalArgs : "");
-			ProcessesToWatch.Add(FirstProcess);
-			FirstProcess.Start();
-			var ProcName = FirstProcess.ProcessName;
+			var bIsSafari = bpath.Contains("Safari");
 
 			if (bIsSafari)
 			{
@@ -100,12 +382,30 @@ namespace HTML5LaunchHelper
 				Proc.StartInfo.UseShellExecute = false;
 				Proc.StartInfo.RedirectStandardOutput = true;
 				Proc.StartInfo.RedirectStandardInput = true;
-				Proc.StartInfo.Arguments = String.Format("-e 'tell application \"Safari\" to open location \"{1}\"'", FirstProcess.Id, Url);
+				Proc.StartInfo.Arguments = String.Format("-e 'tell application \"Safari\" to open location \"{1}\"'", Result.Id, args);
 				Proc.EnableRaisingEvents = true;
 				Proc.Start();
-				Proc.WaitForExit();	
+				Proc.WaitForExit();
 			}
 
+			System.Console.WriteLine("Spawning Browser Process {0} with args {1}\n", bpath, args);
+			return Result;
+
+		}
+
+		static void SpawnBrowserAndBlock(Arguments Args)
+		{
+			// Browsers can be multiprocess programs (Chrome, basically)
+			bool bMultiprocessBrowser = Args.Browser.Contains("chrome");
+			// so we need to catch spawning of other child processes. The trick is
+			// they aren't really child-processes at all. There appears to be no real binding between the two,
+			// So we kind of fudged it a bit here.
+			var PrevProcesses = Process.GetProcesses();
+			var FirstProcess = SpawnBrowserProcess(Args.Browser, Args.BrowserCommandLine);
+			ProcessesToWatch.Add(FirstProcess);
+			FirstProcess.Start();
+			var ProcName = FirstProcess.ProcessName;
+
 			// We should now have a list of processes to watch to exit.
 			// Loop over the calling WaitForExit() until the list is empty.
 			while (ProcessesToWatch.Count() > 0)
@@ -133,14 +433,16 @@ namespace HTML5LaunchHelper
 						}
 						if (!bWasAlive)
 						{
-							try {
+							try
+							{
 								if (!item.HasExited && item.ProcessName.StartsWith(ProcName))
 								{
 									var PID = item.Id;
 									System.Console.WriteLine("Found Process {0} with PID {1} which started at {2}. Waiting on that process to end.", item.ProcessName, item.Id, item.StartTime.ToString());
 									ProcessesToWatch.Add(item);
 								}
-							} catch {}
+							}
+							catch { }
 						}
 					}
 					FirstProcess = null;
@@ -153,7 +455,6 @@ namespace HTML5LaunchHelper
 					{
 						System.Console.WriteLine("A spawned thread has died. Do you have a python server instance running?");
 						HardShutdown();
-						return -1;
 					}
 				}
 
@@ -171,8 +472,6 @@ namespace HTML5LaunchHelper
 					proc.WaitForExit();
 				}
 			}
-
-			return 0;
 		}
 
 		static void HardShutdown()
@@ -196,113 +495,42 @@ namespace HTML5LaunchHelper
 			}
 		}
 
-		static void BeginPythonServer()
+		static int Main(string[] args)
 		{
-			var Result = new Process();
-            Result.StartInfo.FileName = PythonPath;
-			Result.StartInfo.UseShellExecute = false;
-			Result.StartInfo.RedirectStandardOutput = true;
-			Result.StartInfo.RedirectStandardInput = true;
-			Result.StartInfo.Arguments = String.Format("-m SimpleHTTPServer {0}", PythonServerPort);
-			Result.StartInfo.WorkingDirectory = ServerDirectory;
-			Result.EnableRaisingEvents = true;
-			Result.Start();
-			System.Console.Write("Started New Process (PID={0} Name={1})\n", Result.Id, Result.ProcessName);
-			ProcessesToKill.Add(Result);
-		}
+			var Args = new Arguments();
+			if (Args.Parse(args))
+			{
+				Args.ShowParsedValues();
+			}
+			else
+			{
+				System.Console.WriteLine("Incorrect Command line Options.. Exiting");
+				Args.ShowAllOptions();
+				return 0; 
+			}
 
-		static bool ParseArgs(string[] args)
-		{
-			bool bOptionalPythonName = false;
-			bool bParseBrowserPath = false;
-			bool bParseServerDir = false;
-			bool bParseServerPort = false;
+			var Server = new HttpServer(Convert.ToInt32(Args.ServerPort),Args.ServerRoot);
+			Server.Run();
 
-			foreach (var arg in args)
+			if ( Args.Browser != "" )
 			{
-				var a = arg.Trim();
-				if (bOptionalPythonName || bParseBrowserPath || bParseServerDir || bParseServerPort)
-				{
-					if (!a.StartsWith("-"))
-					{
-						if (bOptionalPythonName)
-						{
-							PythonPath = a;
-						}
-						else if (bParseBrowserPath)
-						{
-							BrowserPath = a;
-						}
-						else if (bParseServerDir)
-						{
-							ServerDirectory = a;
-						}
-						else if (bParseServerPort)
-						{
-							PythonServerPort = Convert.ToInt32(a);
-						}
-					}
-					bOptionalPythonName = bParseBrowserPath = bParseServerDir = bParseServerPort = false;
-				}
-				else if (a == "-m")
+				if(!File.Exists(Args.Browser))
 				{
-					bMultiprocessBrowser = true;
-				}
-				else if (a == "-p")
-				{
-					bOptionalPythonName = true;
-				}
-				else if (a == "-h")
-				{
-					bStartPythonServer = true;
-				}
-				else if (a == "-b")
-				{
-					bParseBrowserPath = true;
-				}
-				else if (a == "-w")
-				{
-					bParseServerDir = true;
-				}
-				else if (a == "-s")
-				{
-					bParseServerPort = true;
-				}
-				else if (string.IsNullOrEmpty(Url))
-				{
-					Url = a;
-				}
-				else
-				{
-					if (a.StartsWith ("-") && a.Contains('=') && a.Contains(' ')) 
-					{
-						var Parts = a.Split('=');
-						a = "";
-						for (var p = 0; p < Parts.Count(); ++p)
-						{
-							if (Parts [p].Contains(' '))
-							{
-								a += "\""+Parts[p]+"\"";
-							}
-							else
-							{
-								a += Parts[p];
-							}
-							if (p < Parts.Count() - 1)
-							{
-								a += '=';
-							}
-						}
-					}
-					else if (a.Contains(' '))
-					{
-						a = "\""+a+"\"";
-					}
-					AdditionalArgs += a + " ";
+					System.Console.WriteLine("Browser Not found, Please check -Browser= option");
+					return 0;
 				}
+				SpawnBrowserAndBlock(Args);
+			}
+			else
+			{
+				System.Console.WriteLine("Press Any key Quit Server");
+				System.Console.ReadKey();
 			}
 
-			return !string.IsNullOrEmpty(BrowserPath) && !string.IsNullOrEmpty(Url) && !string.IsNullOrEmpty(ServerDirectory);
+			Server.Stop();
+			return 0;
 		}
+
+
 	}
 }
