Fix incorrect URLs on html5 webserver.

diff --git a/Engine/Source/Programs/HTML5/HTML5LaunchHelper/Program.cs b/Engine/Source/Programs/HTML5/HTML5LaunchHelper/Program.cs
--- a/Engine/Source/Programs/HTML5/HTML5LaunchHelper/Program.cs
+++ b/Engine/Source/Programs/HTML5/HTML5LaunchHelper/Program.cs
@@ -138,8 +138,9 @@ namespace HTML5LaunchHelper
 
 				foreach (string fileName in fileEntries)
 				{
-					string Slash = Directory.Exists(fileName) ? "/" : "" ; 
-					Response += "<tr>\n\t<td><a href=\"" + fileName + "\">" + Path.GetFileName(fileName) + Slash + "</a></td>\n\t<td> " + File.GetLastAccessTime(fileName).ToString() + "</td>\n</tr>\n";
+					string Slash = Directory.Exists(fileName) ? "/" : "" ;
+					string Url = fileName.Replace(Root, WebServer.Prefixes.First().ToString());
+					Response += "<tr>\n\t<td><a href=\"" + Url + "\">" + Path.GetFileName(fileName) + Slash + "</a></td>\n\t<td> " + File.GetLastAccessTime(fileName).ToString() + "</td>\n</tr>\n";
 				}
 				Response  += "</table></html>";
 				byte[] buf = Encoding.UTF8.GetBytes(Response);
