Fix HTML5LauncherHelper on Mac always failing to find browsers.

diff --git a/Engine/Source/Programs/HTML5/HTML5LaunchHelper/Program.cs b/Engine/Source/Programs/HTML5/HTML5LaunchHelper/Program.cs
--- a/Engine/Source/Programs/HTML5/HTML5LaunchHelper/Program.cs
+++ b/Engine/Source/Programs/HTML5/HTML5LaunchHelper/Program.cs
@@ -515,7 +515,7 @@ namespace HTML5LaunchHelper
 
 			if ( Args.Browser != "" )
 			{
-				if(!File.Exists(Args.Browser))
+				if ((!File.Exists(Args.Browser) && !IsRunningOnMac()) || (!Directory.Exists(Args.Browser) && IsRunningOnMac()))
 				{
 					System.Console.WriteLine("Browser Not found, Please check -Browser= option");
 					return 0;
