Fix HTML5LaunchHelper not serving to remote machines.

diff --git a/Engine/Source/Programs/AutomationTool/HTML5/HTML5Platform.Automation.cs b/Engine/Source/Programs/AutomationTool/HTML5/HTML5Platform.Automation.cs
--- a/Engine/Source/Programs/AutomationTool/HTML5/HTML5Platform.Automation.cs
+++ b/Engine/Source/Programs/AutomationTool/HTML5/HTML5Platform.Automation.cs
@@ -404,7 +404,7 @@ public class HTML5Platform : Platform
 		}
 		else
 		{
-			url = String.Format("http://localhost:{0}/{1}", ServerPort, url);
+			url = String.Format("http://{2}:{0}/{1}", ServerPort, url, Environment.MachineName);
 		}
 
 		// Check HTML5LaunchHelper source for command line args
diff --git a/Engine/Source/Programs/HTML5/HTML5LaunchHelper/Program.cs b/Engine/Source/Programs/HTML5/HTML5LaunchHelper/Program.cs
--- a/Engine/Source/Programs/HTML5/HTML5LaunchHelper/Program.cs
+++ b/Engine/Source/Programs/HTML5/HTML5LaunchHelper/Program.cs
@@ -78,12 +78,12 @@ namespace HTML5LaunchHelper
 #endregion
 
 		private HttpListener WebServer = new HttpListener();
-		private string Root; 
+		private string Root;
 
 		public HttpServer(int Port, string ServerRoot)
 		{
 			Root = ServerRoot;
-			WebServer.Prefixes.Add("http://localhost:" + Port.ToString() + "/");
+			WebServer.Prefixes.Add(string.Format("http://{0}:{1}/", Environment.MachineName, Port.ToString()));
 		}
 
 		public void Run() 
