Added extra WebServer Prefixes so localhost, 127.0.0.1, ip address, etc work from HTML5LauncherHelper.

diff --git a/Engine/Source/Programs/HTML5/HTML5LaunchHelper/Program.cs b/Engine/Source/Programs/HTML5/HTML5LaunchHelper/Program.cs
--- a/Engine/Source/Programs/HTML5/HTML5LaunchHelper/Program.cs
+++ b/Engine/Source/Programs/HTML5/HTML5LaunchHelper/Program.cs
@@ -1,5 +1,6 @@
 ﻿using System;
-using System.Net; 
+using System.Net;
+using System.Net.Sockets;
 using System.Collections.Generic;
 using System.Linq;
 using System.Text;
@@ -83,7 +84,17 @@ namespace HTML5LaunchHelper
 		public HttpServer(int Port, string ServerRoot)
 		{
 			Root = ServerRoot;
+			WebServer.Prefixes.Add(string.Format("http://localhost:{0}/", Port.ToString()));
+			WebServer.Prefixes.Add(string.Format("http://127.0.0.1:{0}/", Port.ToString()));
 			WebServer.Prefixes.Add(string.Format("http://{0}:{1}/", Environment.MachineName, Port.ToString()));
+			IPHostEntry host = Dns.GetHostEntry(Dns.GetHostName());
+			foreach (IPAddress ip in host.AddressList)
+			{
+				if (ip.AddressFamily == AddressFamily.InterNetwork)
+				{
+					WebServer.Prefixes.Add(string.Format("http://{0}:{1}/", ip.ToString(), Port.ToString()));
+				}
+			}
 		}
 
 		public void Run() 
