Fix HTML5Launcher to correctly launch Safari and only attach to localhost by default.  (attaching to other interfaces is behind a command line argument now - existing code causes issues on Windows with UAC and on stock Mac)

diff --git a/Engine/Source/Programs/AutomationTool/HTML5/HTML5Platform.Automation.cs b/Engine/Source/Programs/AutomationTool/HTML5/HTML5Platform.Automation.cs
--- a/Engine/Source/Programs/AutomationTool/HTML5/HTML5Platform.Automation.cs
+++ b/Engine/Source/Programs/AutomationTool/HTML5/HTML5Platform.Automation.cs
@@ -363,7 +363,7 @@ public class HTML5Platform : Platform
 		}
 		else
 		{
-			url = String.Format("http://{2}:{0}/{1}", ServerPort, url, Environment.MachineName);
+			url = String.Format("http://localhost:{0}/{1}", ServerPort, url);
 		}
 
 		// Check HTML5LaunchHelper source for command line args
diff --git a/Engine/Source/Programs/HTML5/HTML5LaunchHelper/Program.cs b/Engine/Source/Programs/HTML5/HTML5LaunchHelper/Program.cs
--- a/Engine/Source/Programs/HTML5/HTML5LaunchHelper/Program.cs
+++ b/Engine/Source/Programs/HTML5/HTML5LaunchHelper/Program.cs
@@ -81,18 +81,20 @@ namespace HTML5LaunchHelper
 		private HttpListener WebServer = new HttpListener();
 		private string Root;
 
-		public HttpServer(int Port, string ServerRoot)
+		public HttpServer(int Port, string ServerRoot, bool UseAllPrefixes)
 		{
 			Root = ServerRoot;
 			WebServer.Prefixes.Add(string.Format("http://localhost:{0}/", Port.ToString()));
-			WebServer.Prefixes.Add(string.Format("http://127.0.0.1:{0}/", Port.ToString()));
-			WebServer.Prefixes.Add(string.Format("http://{0}:{1}/", Environment.MachineName, Port.ToString()));
-			IPHostEntry host = Dns.GetHostEntry(Dns.GetHostName());
-			foreach (IPAddress ip in host.AddressList)
-			{
-				if (ip.AddressFamily == AddressFamily.InterNetwork)
-				{
-					WebServer.Prefixes.Add(string.Format("http://{0}:{1}/", ip.ToString(), Port.ToString()));
+
+			if (UseAllPrefixes) 
+			{ 
+				WebServer.Prefixes.Add (string.Format ("http://127.0.0.1:{0}/", Port.ToString ()));
+				WebServer.Prefixes.Add (string.Format ("http://{0}:{1}/", Environment.MachineName, Port.ToString ()));
+				IPHostEntry host = Dns.GetHostEntry (Dns.GetHostName ());
+				foreach (IPAddress ip in host.AddressList) {
+					if (ip.AddressFamily == AddressFamily.InterNetwork) {
+						WebServer.Prefixes.Add (string.Format ("http://{0}:{1}/", ip.ToString (), Port.ToString ()));
+					}
 				}
 			}
 		}
@@ -247,6 +249,14 @@ namespace HTML5LaunchHelper
 			set;
 		}
 
+		[ArgumentName("-UseAllPrefixes="), DefaultArgument("FALSE")]
+		public string UseAllPrefixes
+		{
+			get;
+			set;
+		}
+
+
 		public Arguments()
 		{}
 
@@ -363,14 +373,17 @@ namespace HTML5LaunchHelper
 
 		static Process SpawnBrowserProcess(string bpath, string args)
 		{
+			var bIsSafari = bpath.Contains("Safari");
+
 			var Result = new Process();
 			if (IsRunningOnMac())
 			{
+				string BrowserArgs = bIsSafari ? "" : args; 
 	            Result.StartInfo.FileName = "/usr/bin/open";
 				Result.StartInfo.UseShellExecute = false;
 				Result.StartInfo.RedirectStandardOutput = true;
 				Result.StartInfo.RedirectStandardInput = true;
-				Result.StartInfo.Arguments = String.Format("-nW \"{0}\" --args {1}", bpath, args);
+				Result.StartInfo.Arguments = String.Format("-nW \"{0}\" --args {1}", bpath, BrowserArgs);
 				Result.EnableRaisingEvents = true;
 			}
 			else
@@ -382,24 +395,25 @@ namespace HTML5LaunchHelper
 				Result.StartInfo.Arguments = args;
 				Result.EnableRaisingEvents = true;
 			}
-
-			var bIsSafari = bpath.Contains("Safari");
+			
+			
+			Result.Start(); 
 
 			if (bIsSafari)
 			{
 				// Give Safari time to open...
-				System.Threading.Thread.Sleep(2000);
+				System.Threading.Thread.Sleep(1500);
 				var Proc = new Process();
 				Proc.StartInfo.FileName = "/usr/bin/osascript";
 				Proc.StartInfo.UseShellExecute = false;
 				Proc.StartInfo.RedirectStandardOutput = true;
 				Proc.StartInfo.RedirectStandardInput = true;
-				Proc.StartInfo.Arguments = String.Format("-e 'tell application \"Safari\" to open location \"{1}\"'", Result.Id, args);
+				Proc.StartInfo.Arguments = String.Format("-e 'tell application \"Safari\" to open location \"{0}\"'",  args);
 				Proc.EnableRaisingEvents = true;
 				Proc.Start();
 				Proc.WaitForExit();
 			}
-
+				
 			System.Console.WriteLine("Spawning Browser Process {0} with args {1}\n", bpath, args);
 			return Result;
 
@@ -415,7 +429,6 @@ namespace HTML5LaunchHelper
 			var PrevProcesses = Process.GetProcesses();
 			var FirstProcess = SpawnBrowserProcess(Args.Browser, Args.BrowserCommandLine);
 			ProcessesToWatch.Add(FirstProcess);
-			FirstProcess.Start();
 			var ProcName = FirstProcess.ProcessName;
 
 			// We should now have a list of processes to watch to exit.
@@ -521,7 +534,7 @@ namespace HTML5LaunchHelper
 				return 0; 
 			}
 
-			var Server = new HttpServer(Convert.ToInt32(Args.ServerPort),Args.ServerRoot);
+			var Server = new HttpServer(Convert.ToInt32(Args.ServerPort),Args.ServerRoot, Args.UseAllPrefixes == "FALSE" ? false : true );
 			Server.Run();
 
 			if ( Args.Browser != "" )
