Fix HTML5Launcher to give correct Status codes when request file is not found.
#html5

diff --git a/Engine/Source/Programs/HTML5/HTML5LaunchHelper/Program.cs b/Engine/Source/Programs/HTML5/HTML5LaunchHelper/Program.cs
--- a/Engine/Source/Programs/HTML5/HTML5LaunchHelper/Program.cs
+++ b/Engine/Source/Programs/HTML5/HTML5LaunchHelper/Program.cs
@@ -161,7 +161,7 @@ namespace HTML5LaunchHelper
 				Context.Response.OutputStream.Write(buf, 0, buf.Length);
 				Context.Response.Close();
 			}
-			if (File.Exists(Root + Context.Request.Url.LocalPath))
+			else if (File.Exists(Root + Context.Request.Url.LocalPath))
 			{
 				string RequestedFile = Root + Context.Request.Url.LocalPath;
 				System.Console.WriteLine("Serving " + RequestedFile);
@@ -177,7 +177,7 @@ namespace HTML5LaunchHelper
 					// This is the crux of serving pre-compressed files. 
 					if (Extention == ".gz")
 					{
-					  Context.Response.AddHeader("Content-Encoding", "gzip");
+						Context.Response.AddHeader("Content-Encoding", "gzip");
 					}
 
 					byte[] buffer = new byte[source.Length];
@@ -188,7 +188,22 @@ namespace HTML5LaunchHelper
 					Context.Response.Close();
 				}
 			}
-			Context.Response.Close();
+			else
+			{
+				string RequestedFile = Root + Context.Request.Url.LocalPath;
+				System.Console.WriteLine("Not Serving " + RequestedFile);
+				string Response = "<html>\n" +
+									"<body>\n" +
+									"<h2>404 Not Found</h2>\n" +
+									"</html>\n";
+
+				byte[] buf = Encoding.UTF8.GetBytes(Response);
+				Context.Response.ContentLength64 = buf.Length;
+				Context.Response.StatusCode = (int)HttpStatusCode.NotFound;
+				Context.Response.ContentType = "text/html";
+				Context.Response.OutputStream.Write(buf, 0, buf.Length);
+				Context.Response.Close();
+			}
 		}
 
 		public void Stop()
