HTML5LaunchHelper.exe on linux - //UE4/Main

- it seems that i need to also check-in the exe and pdb file instead of having CIS make and checking-in them itself...
- modified c# program to output a version number to help track which version of HTML5LaunchHelper is running...

#rnx

diff --git a/Engine/Source/Programs/HTML5/HTML5LaunchHelper/Program.cs b/Engine/Source/Programs/HTML5/HTML5LaunchHelper/Program.cs
--- a/Engine/Source/Programs/HTML5/HTML5LaunchHelper/Program.cs
+++ b/Engine/Source/Programs/HTML5/HTML5LaunchHelper/Program.cs
@@ -546,6 +546,7 @@ namespace HTML5LaunchHelper
 
 		static int Main(string[] args)
 		{
+            System.Console.WriteLine("Version: 20170623"); // date: YYYYMMDD - needed to help figure out what version QA is running...
 			var Args = new Arguments();
 			if (Args.Parse(args))
 			{
